Add editor for plugin metadata.

diff --git a/Engine/Plugins/Editor/PluginBrowser/Source/PluginBrowser/PluginBrowser.Build.cs b/Engine/Plugins/Editor/PluginBrowser/Source/PluginBrowser/PluginBrowser.Build.cs
--- a/Engine/Plugins/Editor/PluginBrowser/Source/PluginBrowser/PluginBrowser.Build.cs
+++ b/Engine/Plugins/Editor/PluginBrowser/Source/PluginBrowser/PluginBrowser.Build.cs
@@ -22,6 +22,8 @@ namespace UnrealBuildTool.Rules
 					"EditorStyle",
 					"Projects",
 					"UnrealEd",
+					"PropertyEditor",
+					"SharedSettingsWidgets"
 				}
 			);
 
diff --git a/Engine/Plugins/Editor/PluginBrowser/Source/PluginBrowser/Private/PluginBrowserModule.cpp b/Engine/Plugins/Editor/PluginBrowser/Source/PluginBrowser/Private/PluginBrowserModule.cpp
new file mode 100644
--- /dev/null
+++ b/Engine/Plugins/Editor/PluginBrowser/Source/PluginBrowser/Private/PluginBrowserModule.cpp
@@ -0,0 +1,65 @@
+// Copyright 1998-2015 Epic Games, Inc. All Rights Reserved.
+
+#include "PluginsEditorPrivatePCH.h"
+#include "SPluginsEditor.h"
+#include "Runtime/Core/Public/Features/IModularFeatures.h"
+#include "Editor/UnrealEd/Public/Features/EditorFeatures.h"
+#include "PluginMetadataObject.h"
+#include "PluginStyle.h"
+#include "PropertyEditorModule.h"
+#include "PluginBrowserModule.h"
+#include "SDockTab.h"
+
+#define LOCTEXT_NAMESPACE "PluginsEditor"
+
+IMPLEMENT_MODULE( FPluginBrowserModule, PluginBrowserModule )
+
+const FName FPluginBrowserModule::PluginsEditorTabName( TEXT( "PluginsEditor" ) );
+
+void FPluginBrowserModule::StartupModule()
+{
+	FPluginStyle::Initialize();
+
+	// Register ourselves as an editor feature
+	IModularFeatures::Get().RegisterModularFeature( EditorFeatures::PluginsEditor, this );
+
+	// @todo plugedit: Should we auto-watch plugin directories and refresh the plugins editor UI with newly-dropped plugins?
+
+	struct Local
+	{
+		static TSharedRef<SDockTab> SpawnPluginsEditorTab( const FSpawnTabArgs& SpawnTabArgs )
+		{
+			const TSharedRef<SDockTab> MajorTab = 
+				SNew( SDockTab )
+				.Icon( FPluginStyle::Get()->GetBrush("Plugins.TabIcon") )
+				.TabRole( ETabRole::MajorTab );
+
+			MajorTab->SetContent( SNew( SPluginsEditor ) );
+
+			return MajorTab;
+		}
+	};
+
+	// Register the detail customization for the metadata object
+	FPropertyEditorModule& PropertyModule = FModuleManager::LoadModuleChecked<FPropertyEditorModule>("PropertyEditor");
+	PropertyModule.RegisterCustomClassLayout(UPluginMetadataObject::StaticClass()->GetFName(), FOnGetDetailCustomizationInstance::CreateStatic(&FPluginMetadataCustomization::MakeInstance));
+	
+	// Register a tab spawner so that our tab can be automatically restored from layout files
+	FGlobalTabmanager::Get()->RegisterTabSpawner( PluginsEditorTabName, FOnSpawnTab::CreateStatic( &Local::SpawnPluginsEditorTab ) )
+			.SetDisplayName( LOCTEXT( "PluginsEditorTabTitle", "Plugins" ) )
+			.SetTooltipText( LOCTEXT( "PluginsEditorTooltipText", "Open the Plugins Browser tab." ) )
+			.SetIcon(FSlateIcon(FPluginStyle::Get()->GetStyleSetName(), "Plugins.TabIcon"));
+}
+
+void FPluginBrowserModule::ShutdownModule()
+{
+	FPluginStyle::Shutdown();
+
+	// Unregister the tab spawner
+	FGlobalTabmanager::Get()->UnregisterTabSpawner( PluginsEditorTabName );
+
+	// Unregister our feature
+	IModularFeatures::Get().UnregisterModularFeature( EditorFeatures::PluginsEditor, this );
+}
+
+#undef LOCTEXT_NAMESPACE
\ No newline at end of file
diff --git a/Engine/Plugins/Editor/PluginBrowser/Source/PluginBrowser/Private/PluginBrowserModule.h b/Engine/Plugins/Editor/PluginBrowser/Source/PluginBrowser/Private/PluginBrowserModule.h
new file mode 100644
--- /dev/null
+++ b/Engine/Plugins/Editor/PluginBrowser/Source/PluginBrowser/Private/PluginBrowserModule.h
@@ -0,0 +1,25 @@
+// Copyright 1998-2015 Epic Games, Inc. All Rights Reserved.
+
+#pragma once 
+
+#include "IPluginManager.h"
+#include "ModuleManager.h"
+
+class FPluginBrowserModule : public IPluginsEditor
+{
+public:
+	TMap<IPlugin*, bool> PendingEnablePlugins;
+
+	/** Accessor for the module interface */
+	static FPluginBrowserModule& Get()
+	{
+		return FModuleManager::Get().GetModuleChecked<FPluginBrowserModule>(TEXT("PluginBrowser"));
+	}
+
+	/** IModuleInterface implementation */
+	virtual void StartupModule() override;
+	virtual void ShutdownModule() override;
+
+	/** ID name for the plugins editor major tab */
+	static const FName PluginsEditorTabName;
+};
diff --git a/Engine/Plugins/Editor/PluginBrowser/Source/PluginBrowser/Private/PluginCategoryTreeItem.h b/Engine/Plugins/Editor/PluginBrowser/Source/PluginBrowser/Private/PluginCategoryTreeItem.h
--- a/Engine/Plugins/Editor/PluginBrowser/Source/PluginBrowser/Private/PluginCategoryTreeItem.h
+++ b/Engine/Plugins/Editor/PluginBrowser/Source/PluginBrowser/Private/PluginCategoryTreeItem.h
@@ -20,12 +20,6 @@ public:
 		return ParentCategoryWeak.Pin();
 	}
 
-	/** @return Returns the path up to this category */
-	const FString& GetCategoryPath() const
-	{
-		return CategoryPath;
-	}
-
 	/** @return	Returns the name of this category */
 	const FString& GetCategoryName() const
 	{
@@ -50,38 +44,32 @@ public:
 	}
 
 	/** @return Returns the plugins in this category */
-	const TArray< TSharedPtr< FPluginStatus > >& GetPlugins() const
+	const TArray< IPlugin* >& GetPlugins() const
 	{
 		return Plugins;
 	}
 
 	/** Adds a plugin to this category */
-	void AddPlugin( const TSharedRef< FPluginStatus >& PluginStatus )
+	void AddPlugin( IPlugin* Plugin )
 	{
-		Plugins.Add( PluginStatus );
+		Plugins.Add( Plugin );
 	}
 
 
 public:
 
 	/** Constructor for FPluginCategoryTreeItem */
-	FPluginCategoryTreeItem(const FPluginCategoryTreeItemPtr ParentCategory, const FString& InitCategoryPath, const FString& InitCategoryName, const FText& InitCategoryDisplayName)
+	FPluginCategoryTreeItem(const FPluginCategoryTreeItemPtr ParentCategory, const FString& InitCategoryName, const FText& InitCategoryDisplayName)
 		: ParentCategoryWeak( ParentCategory ),
-		  CategoryPath( InitCategoryPath ),
 		  CategoryName(InitCategoryName),
 		  CategoryDisplayName(InitCategoryDisplayName)
 	{
 	}
 
 
-private:
-
 	/** Parent category item or NULL if this is a root category */
 	TWeakPtr< FPluginCategoryTreeItem > ParentCategoryWeak;
 
-	/** Category path string.  The full path up to this category */
-	FString CategoryPath;
-
 	/** Name of the category */
 	FString CategoryName;
 
@@ -92,7 +80,7 @@ private:
 	TArray< TSharedPtr< FPluginCategoryTreeItem > > SubCategories;
 
 	/** Plugins in this category */
-	TArray< TSharedPtr< FPluginStatus > > Plugins;
+	TArray< IPlugin* > Plugins;
 
 };
 
diff --git a/Engine/Plugins/Editor/PluginBrowser/Source/PluginBrowser/Private/PluginListItem.h b/Engine/Plugins/Editor/PluginBrowser/Source/PluginBrowser/Private/PluginListItem.h
--- a/Engine/Plugins/Editor/PluginBrowser/Source/PluginBrowser/Private/PluginListItem.h
+++ b/Engine/Plugins/Editor/PluginBrowser/Source/PluginBrowser/Private/PluginListItem.h
@@ -17,7 +17,7 @@ class FPluginListItem
 public:
 
 	/** Plugin status object, which contains information about the plugin such as its name and description */
-	FPluginStatus PluginStatus;
+	IPlugin* Plugin;
 
 
 public:
diff --git a/Engine/Plugins/Editor/PluginBrowser/Source/PluginBrowser/Private/PluginMetadataObject.cpp b/Engine/Plugins/Editor/PluginBrowser/Source/PluginBrowser/Private/PluginMetadataObject.cpp
new file mode 100644
--- /dev/null
+++ b/Engine/Plugins/Editor/PluginBrowser/Source/PluginBrowser/Private/PluginMetadataObject.cpp
@@ -0,0 +1,99 @@
+// Copyright 1998-2015 Epic Games, Inc. All Rights Reserved.
+
+#include "PluginsEditorPrivatePCH.h"
+#include "DetailLayoutBuilder.h"
+#include "DetailCategoryBuilder.h"
+#include "DetailWidgetRow.h"
+#include "SExternalImageReference.h"
+#include "IPluginManager.h"
+#include "PluginMetadataObject.h"
+
+UPluginMetadataObject::UPluginMetadataObject(const FObjectInitializer& ObjectInitializer)
+{
+}
+
+void UPluginMetadataObject::PopulateFromDescriptor(const FPluginDescriptor& InDescriptor)
+{
+	Version = InDescriptor.Version;
+	VersionName = InDescriptor.VersionName;
+	FriendlyName = InDescriptor.FriendlyName;
+	Description = InDescriptor.Description;
+	Category = InDescriptor.Category;
+	CreatedBy = InDescriptor.CreatedBy;
+	CreatedByURL = InDescriptor.CreatedByURL;
+	DocsURL = InDescriptor.DocsURL;
+	bCanContainContent = InDescriptor.bCanContainContent;
+	bIsBetaVersion = InDescriptor.bIsBetaVersion;
+}
+
+void UPluginMetadataObject::CopyIntoDescriptor(FPluginDescriptor& OutDescriptor)
+{
+	OutDescriptor.Version = Version;
+	OutDescriptor.VersionName = VersionName;
+	OutDescriptor.FriendlyName = FriendlyName;
+	OutDescriptor.Description = Description;
+	OutDescriptor.Category = Category;
+	OutDescriptor.CreatedBy = CreatedBy;
+	OutDescriptor.CreatedByURL = CreatedByURL;
+	OutDescriptor.DocsURL = DocsURL;
+	OutDescriptor.bCanContainContent = bCanContainContent;
+	OutDescriptor.bIsBetaVersion = bIsBetaVersion;
+}
+
+/////////////////////////////////////////////////////////////////////////////////////////////////////////
+
+TSharedRef<IDetailCustomization> FPluginMetadataCustomization::MakeInstance()
+{
+	return MakeShareable(new FPluginMetadataCustomization());
+}
+
+void FPluginMetadataCustomization::CustomizeDetails(IDetailLayoutBuilder& DetailBuilder)
+{
+	TArray<TWeakObjectPtr<UObject>> Objects;
+	DetailBuilder.GetObjectsBeingCustomized(Objects);
+
+	if(Objects.Num() == 1 && Objects[0].IsValid())
+	{
+		UPluginMetadataObject* PluginMetadata = Cast<UPluginMetadataObject>(Objects[0].Get());
+		if(PluginMetadata != nullptr && PluginMetadata->TargetIconPath.Len() > 0)
+		{
+			// Get the current icon path
+			FString CurrentIconPath = PluginMetadata->TargetIconPath;
+			if(!FPaths::FileExists(CurrentIconPath))
+			{
+				CurrentIconPath = IPluginManager::Get().FindPlugin(TEXT("PluginBrowser"))->GetBaseDir() / TEXT("Resources") / TEXT("DefaultIcon128.png");
+			}
+
+			// Add the customization to edit the icon row
+			IDetailCategoryBuilder& ImageCategory = DetailBuilder.EditCategory(TEXT("Icon"));
+			const FText IconDesc(NSLOCTEXT("PluginBrowser", "PluginIcon", "Icon"));
+			ImageCategory.AddCustomRow(IconDesc)
+			.NameContent()
+			[
+				SNew(SHorizontalBox)
+				+SHorizontalBox::Slot()
+				.Padding( FMargin( 0, 1, 0, 1 ) )
+				.FillWidth(1.0f)
+				[
+					SNew(STextBlock)
+					.Text(IconDesc)
+					.Font(DetailBuilder.GetDetailFont())
+				]
+			]
+			.ValueContent()
+			.MaxDesiredWidth(500.0f)
+			.MinDesiredWidth(100.0f)
+			[
+				SNew(SHorizontalBox)
+				+SHorizontalBox::Slot()
+				.FillWidth(1.0f)
+				.VAlign(VAlign_Center)
+				[
+					SNew(SExternalImageReference, CurrentIconPath, PluginMetadata->TargetIconPath)
+					.FileDescription(IconDesc)
+					.RequiredSize(FIntPoint(128, 128))
+				]
+			];
+		}
+	}
+}
diff --git a/Engine/Plugins/Editor/PluginBrowser/Source/PluginBrowser/Private/PluginMetadataObject.h b/Engine/Plugins/Editor/PluginBrowser/Source/PluginBrowser/Private/PluginMetadataObject.h
new file mode 100644
--- /dev/null
+++ b/Engine/Plugins/Editor/PluginBrowser/Source/PluginBrowser/Private/PluginMetadataObject.h
@@ -0,0 +1,86 @@
+// Copyright 1998-2015 Epic Games, Inc. All Rights Reserved.
+
+#pragma once
+
+#include "PluginDescriptor.h"
+#include "IDetailCustomization.h"
+#include "PluginMetadataObject.generated.h"
+
+/**
+*	We use this object to display plugin properties using details view.
+*/
+UCLASS()
+class UPluginMetadataObject : public UObject
+{
+public:
+	GENERATED_BODY()
+
+	/* Default constructor */
+	UPluginMetadataObject(const FObjectInitializer& ObjectInitializer);
+
+	/** Path to this this plugin's icon */
+	FString TargetIconPath;
+
+	/** Version number for the plugin.  The version number must increase with every version of the plugin, so that the system
+	can determine whether one version of a plugin is newer than another, or to enforce other requirements.  This version
+	number is not displayed in front-facing UI.  Use the VersionName for that. */
+	UPROPERTY(VisibleAnywhere, Category = "Details")
+	int32 Version;
+
+	/** Name of the version for this plugin.  This is the front-facing part of the version number.  It doesn't need to match
+	the version number numerically, but should be updated when the version number is increased accordingly. */
+	UPROPERTY(EditAnywhere, Category = "Details")
+	FString VersionName;
+
+	/** Friendly name of the plugin */
+	UPROPERTY(EditAnywhere, Category = "Details")
+	FString FriendlyName;
+
+	/** Description of the plugin */
+	UPROPERTY(EditAnywhere, Category = "Details")
+	FString Description;
+
+	/** The category that this plugin belongs to */
+	UPROPERTY(EditAnywhere, Category = "Details")
+	FString Category;
+
+	/** The company or individual who created this plugin.  This is an optional field that may be displayed in the user interface. */
+	UPROPERTY(EditAnywhere, Category = "Details")
+	FString CreatedBy;
+
+	/** Hyperlink URL string for the company or individual who created this plugin.  This is optional. */
+	UPROPERTY(EditAnywhere, Category = "Details")
+	FString CreatedByURL;
+
+	/** Documentation URL string. */
+	UPROPERTY(EditAnywhere, Category = "Details")
+	FString DocsURL;
+
+	/** Can this plugin contain content? */
+	UPROPERTY(EditAnywhere, Category = "Details")
+	bool bCanContainContent;
+
+	/** Marks the plugin as beta in the UI */
+	UPROPERTY(EditAnywhere, Category = "Details")
+	bool bIsBetaVersion;
+
+	/**
+	 * Populate the fields of this object from an existing descriptor.
+	 */
+	void PopulateFromDescriptor(const FPluginDescriptor& InDescriptor);
+
+	/**
+	 * Copy the metadata fields into a plugin descriptor.
+	 */
+	void CopyIntoDescriptor(FPluginDescriptor& OutDescriptor);
+};
+
+/**
+ * Detail customization to allow editing the plugin's icon
+ */
+class FPluginMetadataCustomization : public IDetailCustomization
+{
+public:
+	static TSharedRef<IDetailCustomization> MakeInstance();
+	virtual void CustomizeDetails(IDetailLayoutBuilder& DetailBuilder) override;
+};
diff --git a/Engine/Plugins/Editor/PluginBrowser/Source/PluginBrowser/Private/PluginStyle.cpp b/Engine/Plugins/Editor/PluginBrowser/Source/PluginBrowser/Private/PluginStyle.cpp
--- a/Engine/Plugins/Editor/PluginBrowser/Source/PluginBrowser/Private/PluginStyle.cpp
+++ b/Engine/Plugins/Editor/PluginBrowser/Source/PluginBrowser/Private/PluginStyle.cpp
@@ -164,6 +164,11 @@ void FPluginStyle::Initialize()
 
 			StyleSet->Set( "PluginTile.BetaWarning", new IMAGE_BRUSH( "icon_plugins_betawarn_14px", FVector2D(14, 14) ) );
 		}
+
+		// Metadata editor
+		{
+			StyleSet->Set("PluginMetadataNameFont", TTF_FONT("Fonts/Roboto-Bold", 18));
+		}
 	}
 
 	FSlateStyleRegistry::RegisterSlateStyle( *StyleSet.Get() );
diff --git a/Engine/Plugins/Editor/PluginBrowser/Source/PluginBrowser/Private/PluginsEditor.cpp b/Engine/Plugins/Editor/PluginBrowser/Source/PluginBrowser/Private/PluginsEditor.cpp
deleted file mode 100644
diff --git a/Engine/Plugins/Editor/PluginBrowser/Source/PluginBrowser/Private/PluginsEditorPrivatePCH.h b/Engine/Plugins/Editor/PluginBrowser/Source/PluginBrowser/Private/PluginsEditorPrivatePCH.h
--- a/Engine/Plugins/Editor/PluginBrowser/Source/PluginBrowser/Private/PluginsEditorPrivatePCH.h
+++ b/Engine/Plugins/Editor/PluginBrowser/Source/PluginBrowser/Private/PluginsEditorPrivatePCH.h
@@ -4,7 +4,8 @@
 
 #include "UnrealEd.h"
 #include "IPluginsEditor.h"
+#include "IPluginManager.h"
 
 #include "TextFilter.h"
 
-typedef TTextFilter< const struct FPluginStatus& > FPluginTextFilter;
+typedef TTextFilter< const IPlugin* > FPluginTextFilter;
diff --git a/Engine/Plugins/Editor/PluginBrowser/Source/PluginBrowser/Private/SPluginCategories.cpp b/Engine/Plugins/Editor/PluginBrowser/Source/PluginBrowser/Private/SPluginCategories.cpp
--- a/Engine/Plugins/Editor/PluginBrowser/Source/PluginBrowser/Private/SPluginCategories.cpp
+++ b/Engine/Plugins/Editor/PluginBrowser/Source/PluginBrowser/Private/SPluginCategories.cpp
@@ -15,7 +15,13 @@ void SPluginCategories::Construct( const FArguments& Args, const TSharedRef< SPl
 {
 	OwnerWeak = Owner;
 
-	RebuildAndFilterCategoryTree();
+	RootPluginCategories.Reset();
+
+	// Add the root level categories
+	TSharedRef<FPluginCategoryTreeItem> InstalledCategory = MakeShareable(new FPluginCategoryTreeItem(NULL, TEXT("Installed"), LOCTEXT("InstalledCategoryName", "Installed")));
+	RootPluginCategories.Add( InstalledCategory );
+	TSharedRef<FPluginCategoryTreeItem> ProjectCategory = MakeShareable(new FPluginCategoryTreeItem(NULL, TEXT("Project"), LOCTEXT("ProjectCategoryName", "Project")));
+	RootPluginCategories.Add( ProjectCategory );
 
 	PluginCategoryTreeView =
 		SNew( SPluginCategoryTreeView )
@@ -38,15 +44,7 @@ void SPluginCategories::Construct( const FArguments& Args, const TSharedRef< SPl
 		PluginCategoryTreeView->SetItemExpansion( Category, true );
 	}
 
-	// Select the first item by default
-	if( RootPluginCategories.Num() > 0 )
-	{
-		TArray<FPluginCategoryTreeItemPtr> SubCategories = RootPluginCategories[0]->GetSubCategories();
-		if(SubCategories.Num() > 0)
-		{
-			PluginCategoryTreeView->SetSelection( SubCategories[ 0 ] );
-		}
-	}
+	RebuildAndFilterCategoryTree();
 
 	ChildSlot.AttachWidget( PluginCategoryTreeView.ToSharedRef() );
 }
@@ -66,118 +64,109 @@ SPluginsEditor& SPluginCategories::GetOwner()
 
 void SPluginCategories::RebuildAndFilterCategoryTree()
 {
+	// Get a plugin from the currently selected category, so we can track it if it's removed
+	IPlugin* TrackPlugin = nullptr;
+	for(FPluginCategoryTreeItemPtr SelectedItem: PluginCategoryTreeView->GetSelectedItems())
 	{
-		RootPluginCategories.Reset();
-
-		// Add the root level categories
-		TSharedRef<FPluginCategoryTreeItem> InstalledCategory = MakeShareable(new FPluginCategoryTreeItem(NULL, TEXT("Installed"), TEXT("Installed"), LOCTEXT("InstalledCategoryName", "Installed")));
-		RootPluginCategories.Add( InstalledCategory );
-		TSharedRef<FPluginCategoryTreeItem> ProjectCategory = MakeShareable(new FPluginCategoryTreeItem(NULL, TEXT("Project"), TEXT("Project"), LOCTEXT("ProjectCategoryName", "Project")));
-		RootPluginCategories.Add( ProjectCategory );
+		if(SelectedItem->Plugins.Num() > 0)
+		{
+			TrackPlugin = SelectedItem->Plugins[0];
+			break;
+		}
+	}
 
+	// Clear the list of plugins in each current category
+	for(FPluginCategoryTreeItemPtr RootCategory: RootPluginCategories)
+	{
+		for(FPluginCategoryTreeItemPtr Category: RootCategory->AccessSubCategories())
+		{
+			Category->Plugins.Reset();
+		}
+	}
 
-		const TArray< FPluginStatus > Plugins = IPluginManager::Get().QueryStatusForAllPlugins();
-		for(const FPluginStatus& Plugin: Plugins)
+	// Add all the known plugins into categories
+	FPluginCategoryTreeItemPtr SelectCategory;
+	for(IPlugin* Plugin: IPluginManager::Get().GetDiscoveredPlugins())
+	{
+		// Figure out which base category this plugin belongs in
+		FPluginCategoryTreeItemPtr RootCategory;
+		if( Plugin->GetLoadedFrom() == EPluginLoadedFrom::Engine )
 		{
-			// Figure out which base category this plugin belongs in
-			FPluginCategoryTreeItemPtr CategoryForPlugin;
-			if( Plugin.LoadedFrom == EPluginLoadedFrom::Engine )
-			{
-				CategoryForPlugin = InstalledCategory;
-			}
-			else
-			{
-				CategoryForPlugin = ProjectCategory;
-			}
-			check( CategoryForPlugin.IsValid() );
-			FString ItemCategoryPath = CategoryForPlugin->GetCategoryPath();
+			RootCategory = RootPluginCategories[0];
+		}
+		else
+		{
+			RootCategory = RootPluginCategories[1];
+		}
 
+		// Get the subcategory for this plugin
+		FString CategoryName = Plugin->GetDescriptor().Category;
+		if(CategoryName.IsEmpty())
+		{
+			CategoryName = TEXT("Other");
+		}
 
-			const FString& CategoryPath = Plugin.Descriptor.Category;
+		// Locate this category at the level we're at in the hierarchy
+		FPluginCategoryTreeItemPtr FoundCategory = NULL;
+		for(FPluginCategoryTreeItemPtr TestCategory: RootCategory->SubCategories)
+		{
+			if(TestCategory->GetCategoryName() == CategoryName)
 			{
-				// We're expecting the category string to be in the "A.B.C" format.  We'll split up the string here and form
-				// a proper hierarchy in the UI
-				TArray< FString > SplitCategories;
-				CategoryPath.ParseIntoArray( SplitCategories, TEXT( "." ), true /* bCullEmpty */ );
-
-				// Make sure all of the categories exist
-				for( auto SplitCategoryIt( SplitCategories.CreateConstIterator() ); SplitCategoryIt; ++SplitCategoryIt )
-				{
-					const auto& SplitCategory = *SplitCategoryIt;
-
-					if( !ItemCategoryPath.IsEmpty() )
-					{
-						ItemCategoryPath += TEXT( '.' );
-					}
-					ItemCategoryPath += SplitCategory;
-
-					// Locate this category at the level we're at in the hierarchy
-					FPluginCategoryTreeItemPtr FoundCategory = NULL;
-					TArray< FPluginCategoryTreeItemPtr >& TestCategoryList = CategoryForPlugin.IsValid() ? CategoryForPlugin->AccessSubCategories() : RootPluginCategories;
-					for( auto TestCategoryIt( TestCategoryList.CreateConstIterator() ); TestCategoryIt; ++TestCategoryIt )
-					{
-						const auto& TestCategory = *TestCategoryIt;
-						if( TestCategory->GetCategoryName() == SplitCategory )
-						{
-							// Found it!
-							FoundCategory = TestCategory;
-							break;
-						}
-					}
-
-					if( !FoundCategory.IsValid() )
-					{
-						// OK, this is a new category name for us, so add it now!
-						const auto ParentCategory = CategoryForPlugin;
-
-						//@todo Allow for properly localized category names [3/7/2014 justin.sargent]
-						FoundCategory = MakeShareable(new FPluginCategoryTreeItem(ParentCategory, ItemCategoryPath, SplitCategory, FText::FromString(SplitCategory)));
-						TestCategoryList.Add( FoundCategory );
-					}
-
-					// Descend the hierarchy for the next category
-					CategoryForPlugin = FoundCategory;
-				}
+				FoundCategory = TestCategory;
+				break;
 			}
+		}
 
-			check( CategoryForPlugin.IsValid() );
-			
-			// Associate the plugin with the category
-			// PERFORMANCE NOTE: Copying FPluginStats by value here
-			CategoryForPlugin->AddPlugin( MakeShareable( new FPluginStatus( Plugin ) ) );
+		if( !FoundCategory.IsValid() )
+		{
+			//@todo Allow for properly localized category names [3/7/2014 justin.sargent]
+			FoundCategory = MakeShareable(new FPluginCategoryTreeItem(RootCategory, CategoryName, FText::FromString(CategoryName)));
+			RootCategory->SubCategories.Add( FoundCategory );
 		}
+			
+		// Associate the plugin with the category
+		FoundCategory->AddPlugin(Plugin);
 
+		// Update the selection if this is the plugin we're tracking
+		if(TrackPlugin == Plugin)
+		{
+			SelectCategory = FoundCategory;
+		}
+	}
 
-		// Sort every single category alphabetically
+	// Remove any empty categories, keeping track of which items are still selected
+	for(FPluginCategoryTreeItemPtr RootCategory: RootPluginCategories)
+	{
+		for(int32 Idx = 0; Idx < RootCategory->SubCategories.Num(); Idx++)
 		{
-			struct FPluginCategoryTreeItemSorter
+			if(RootCategory->SubCategories[Idx]->Plugins.Num() == 0)
 			{
-				bool operator()( const FPluginCategoryTreeItemPtr& A, const FPluginCategoryTreeItemPtr& B ) const
-				{
-					return A->GetCategoryDisplayName().CompareTo( B->GetCategoryDisplayName() ) == -1;
-				}
-
-				static void SortCategoriesRecursively( TArray< FPluginCategoryTreeItemPtr >& Categories )
-				{
-					// Sort the categories
-					Categories.Sort( FPluginCategoryTreeItemSorter() );
-
-					// Sort sub-categories
-					for( auto SubCategoryIt( Categories.CreateConstIterator() ); SubCategoryIt; ++SubCategoryIt )
-					{
-						const auto& SubCategory = *SubCategoryIt;
-						SortCategoriesRecursively( SubCategory->AccessSubCategories() );
-					}
-				}
-			};
-
-			FPluginCategoryTreeItemSorter::SortCategoriesRecursively( RootPluginCategories );
+				RootCategory->SubCategories.RemoveAt(Idx);
+			}
 		}
 	}
 
-	if( PluginCategoryTreeView.IsValid() )
+	// Sort every single category alphabetically
+	for(FPluginCategoryTreeItemPtr RootCategory: RootPluginCategories)
+	{
+		RootCategory->SubCategories.Sort([](const FPluginCategoryTreeItemPtr& A, const FPluginCategoryTreeItemPtr& B) -> bool { return A->CategoryName < B->CategoryName; });
+	}
+
+	// Refresh the view
+	PluginCategoryTreeView->RequestTreeRefresh();
+
+	// Make sure we have something selected
+	if(SelectCategory.IsValid())
+	{
+		PluginCategoryTreeView->SetSelection(SelectCategory);
+	}
+	else if(RootPluginCategories[0]->SubCategories.Num() > 0)
 	{
-		PluginCategoryTreeView->RequestTreeRefresh();
+		PluginCategoryTreeView->SetSelection(RootPluginCategories[0]->SubCategories[0]);
+	}
+	else
+	{
+		PluginCategoryTreeView->SetSelection(RootPluginCategories[0]);
 	}
 }
 
@@ -220,7 +209,6 @@ FPluginCategoryTreeItemPtr SPluginCategories::GetSelectedCategory() const
 	return NULL;
 }
 
-
 void SPluginCategories::SelectCategory( const FPluginCategoryTreeItemPtr& CategoryToSelect )
 {
 	if( ensure( PluginCategoryTreeView.IsValid() ) )
@@ -229,11 +217,20 @@ void SPluginCategories::SelectCategory( const FPluginCategoryTreeItemPtr& Catego
 	}
 }
 
-
 bool SPluginCategories::IsItemExpanded( const FPluginCategoryTreeItemPtr Item ) const
 {
 	return PluginCategoryTreeView->IsItemExpanded( Item );
 }
 
+void SPluginCategories::SetNeedsRefresh()
+{
+	RegisterActiveTimer (0.f, FWidgetActiveTimerDelegate::CreateSP (this, &SPluginCategories::TriggerCategoriesRefresh));
+}
+
+EActiveTimerReturnType SPluginCategories::TriggerCategoriesRefresh(double InCurrentTime, float InDeltaTime)
+{
+	RebuildAndFilterCategoryTree();
+	return EActiveTimerReturnType::Stop;
+}
 
 #undef LOCTEXT_NAMESPACE
\ No newline at end of file
diff --git a/Engine/Plugins/Editor/PluginBrowser/Source/PluginBrowser/Private/SPluginCategories.h b/Engine/Plugins/Editor/PluginBrowser/Source/PluginBrowser/Private/SPluginCategories.h
--- a/Engine/Plugins/Editor/PluginBrowser/Source/PluginBrowser/Private/SPluginCategories.h
+++ b/Engine/Plugins/Editor/PluginBrowser/Source/PluginBrowser/Private/SPluginCategories.h
@@ -40,6 +40,8 @@ public:
 	/** @return Returns true if the specified item is currently expanded in the tree */
 	bool IsItemExpanded( const FPluginCategoryTreeItemPtr Item ) const;
 
+	/** Signal that the categories list needs to be refreshed */
+	void SetNeedsRefresh();
 
 private:
 	
@@ -55,6 +57,9 @@ private:
 	/** Rebuilds the category tree from scratch */
 	void RebuildAndFilterCategoryTree();
 
+	/** Callback for refreshing the tree */
+	EActiveTimerReturnType TriggerCategoriesRefresh(double InCurrentTime, float InDeltaTime);
+
 private:
 
 	/** Weak pointer back to its owner */
diff --git a/Engine/Plugins/Editor/PluginBrowser/Source/PluginBrowser/Private/SPluginCategoryTreeItem.cpp b/Engine/Plugins/Editor/PluginBrowser/Source/PluginBrowser/Private/SPluginCategoryTreeItem.cpp
--- a/Engine/Plugins/Editor/PluginBrowser/Source/PluginBrowser/Private/SPluginCategoryTreeItem.cpp
+++ b/Engine/Plugins/Editor/PluginBrowser/Source/PluginBrowser/Private/SPluginCategoryTreeItem.cpp
@@ -80,26 +80,17 @@ const FSlateBrush* SPluginCategoryTreeItem::GetIconBrush() const
 	// @todo plugedit: Need support for styles embedded into plugins!
 	const FSlateBrush* IconBrush;
 
-	if ( ItemData->GetCategoryPath() == TEXT("BuiltIn") )
+	if(ItemData->GetParentCategory().IsValid())
 	{
-		IconBrush = FPluginStyle::Get()->GetBrush( "CategoryTreeItem.BuiltIn" );
+		IconBrush = FPluginStyle::Get()->GetBrush( "CategoryTreeItem.LeafItemWithPlugin" );
 	}
-	else if ( ItemData->GetCategoryPath() == TEXT("Installed") )
+	else if ( ItemData->GetCategoryName() == TEXT("Installed"))
 	{
-		IconBrush = FPluginStyle::Get()->GetBrush( "CategoryTreeItem.Installed" );
+		IconBrush = FPluginStyle::Get()->GetBrush( "CategoryTreeItem.BuiltIn" );
 	}
 	else
 	{
-		const bool bIsLeafItemWithPlugins = ( ItemData->GetSubCategories().Num() == 0 ) && ( ItemData->GetPlugins().Num() > 0 );
-		if( bIsLeafItemWithPlugins )
-		{
-			IconBrush = FPluginStyle::Get()->GetBrush( "CategoryTreeItem.LeafItemWithPlugin" );
-		}
-		else
-		{
-			const bool bIsExpanded = OwnerWeak.Pin()->IsItemExpanded( ItemData );
-			IconBrush = bIsExpanded ? FPluginStyle::Get()->GetBrush( "CategoryTreeItem.ExpandedCategory" ) : FPluginStyle::Get()->GetBrush( "CategoryTreeItem.Category" );
-		}
+		IconBrush = FPluginStyle::Get()->GetBrush( "CategoryTreeItem.Installed" );
 	}
 
 	return IconBrush;
diff --git a/Engine/Plugins/Editor/PluginBrowser/Source/PluginBrowser/Private/SPluginList.cpp b/Engine/Plugins/Editor/PluginBrowser/Source/PluginBrowser/Private/SPluginList.cpp
--- a/Engine/Plugins/Editor/PluginBrowser/Source/PluginBrowser/Private/SPluginList.cpp
+++ b/Engine/Plugins/Editor/PluginBrowser/Source/PluginBrowser/Private/SPluginList.cpp
@@ -68,22 +68,12 @@ TSharedRef<ITableRow> SPluginList::PluginListView_OnGenerateRow( FPluginListItem
 void SPluginList::GetPlugins( const TSharedPtr< FPluginCategoryTreeItem >& Category )
 {
 	// Add plugins from this category
-	const TArray< TSharedPtr< FPluginStatus > >& Plugins = Category->GetPlugins();
-	for( auto PluginIt( Plugins.CreateConstIterator() ); PluginIt; ++PluginIt )
+	for(IPlugin* Plugin: Category->GetPlugins())
 	{
-		const auto& PluginStatus = *( *PluginIt );
-
-		// Skip plugins that don't match the current filter criteria
-		bool bPassesFilter = true;
-		if( !OwnerWeak.Pin()->GetPluginTextFilter().PassesFilter( PluginStatus ) )
-		{
-			bPassesFilter = false;
-		}
-
-		if( bPassesFilter )
+		if(OwnerWeak.Pin()->GetPluginTextFilter().PassesFilter(Plugin))
 		{
 			const TSharedRef< FPluginListItem > NewItem( new FPluginListItem() );
-			NewItem->PluginStatus = PluginStatus;
+			NewItem->Plugin = Plugin;
 			PluginListItems.Add( NewItem );
 		}
 	}
@@ -109,7 +99,7 @@ void SPluginList::RebuildAndFilterPluginList()
 			{
 				bool operator()( const FPluginListItemPtr& A, const FPluginListItemPtr& B ) const
 				{
-					return A->PluginStatus.Descriptor.FriendlyName < B->PluginStatus.Descriptor.FriendlyName;
+					return A->Plugin->GetDescriptor().FriendlyName < B->Plugin->GetDescriptor().FriendlyName;
 				}
 			};
 			PluginListItems.Sort( FPluginListItemSorter() );
diff --git a/Engine/Plugins/Editor/PluginBrowser/Source/PluginBrowser/Private/SPluginListTile.cpp b/Engine/Plugins/Editor/PluginBrowser/Source/PluginBrowser/Private/SPluginListTile.cpp
--- a/Engine/Plugins/Editor/PluginBrowser/Source/PluginBrowser/Private/SPluginListTile.cpp
+++ b/Engine/Plugins/Editor/PluginBrowser/Source/PluginBrowser/Private/SPluginListTile.cpp
@@ -7,8 +7,12 @@
 #include "SPluginList.h"
 #include "PluginStyle.h"
 #include "GameProjectGenerationModule.h"
+#include "IDetailsView.h"
 #include "SHyperlink.h"
+#include "PluginMetadataObject.h"
 #include "IProjectManager.h"
+#include "PluginBrowserModule.h"
+#include "ISourceControlModule.h"
 
 #define LOCTEXT_NAMESPACE "PluginListTile"
 
@@ -27,8 +31,11 @@ void SPluginListTile::Construct( const FArguments& Args, const TSharedRef<SPlugi
 
 	// @todo plugedit: Maybe we should do the FileExists check ONCE at plugin load time and not at query time
 
+	IPlugin* Plugin = ItemData->Plugin;
+	const FPluginDescriptor& PluginDescriptor = Plugin->GetDescriptor();
+
 	// Plugin thumbnail image
-	FString Icon128FilePath = Item->PluginStatus.PluginDirectory / TEXT("Resources/Icon128.png");
+	FString Icon128FilePath = Plugin->GetBaseDir() / TEXT("Resources/Icon128.png");
 	if(!FPlatformFileManager::Get().GetPlatformFile().FileExists(*Icon128FilePath))
 	{
 		Icon128FilePath = IPluginManager::Get().FindPlugin(TEXT("PluginBrowser"))->GetBaseDir() / TEXT("Resources/DefaultIcon128.png");
@@ -44,13 +51,13 @@ void SPluginListTile::Construct( const FArguments& Args, const TSharedRef<SPlugi
 	// create documentation link
 	TSharedPtr<SWidget> DocumentationWidget;
 	{
-		if (Item->PluginStatus.Descriptor.DocsURL.IsEmpty())
+		if (PluginDescriptor.DocsURL.IsEmpty())
 		{
 			DocumentationWidget = SNullWidget::NullWidget;
 		}
 		else
 		{
-			FString DocsURL = Item->PluginStatus.Descriptor.DocsURL;
+			FString DocsURL = PluginDescriptor.DocsURL;
 			DocumentationWidget = SNew(SHorizontalBox)
 
 			+ SHorizontalBox::Slot()
@@ -78,11 +85,11 @@ void SPluginListTile::Construct( const FArguments& Args, const TSharedRef<SPlugi
 	// create vendor link
 	TSharedPtr<SWidget> CreatedByWidget;
 	{
-		if (Item->PluginStatus.Descriptor.CreatedBy.IsEmpty())
+		if (PluginDescriptor.CreatedBy.IsEmpty())
 		{
 			CreatedByWidget = SNullWidget::NullWidget;
 		}
-		else if (Item->PluginStatus.Descriptor.CreatedByURL.IsEmpty())
+		else if (PluginDescriptor.CreatedByURL.IsEmpty())
 		{
 			CreatedByWidget = SNew(SHorizontalBox)
 
@@ -101,12 +108,12 @@ void SPluginListTile::Construct( const FArguments& Args, const TSharedRef<SPlugi
 				.Padding(2.0f, 0.0f, 0.0f, 0.0f)
 				[
 					SNew(STextBlock)
-						.Text(FText::FromString(Item->PluginStatus.Descriptor.CreatedBy))
+						.Text(FText::FromString(PluginDescriptor.CreatedBy))
 				];
 		}
 		else
 		{
-			FString CreatedByURL = Item->PluginStatus.Descriptor.CreatedByURL;
+			FString CreatedByURL = PluginDescriptor.CreatedByURL;
 			CreatedByWidget = SNew(SHorizontalBox)
 
 			+ SHorizontalBox::Slot()
@@ -124,7 +131,7 @@ void SPluginListTile::Construct( const FArguments& Args, const TSharedRef<SPlugi
 				.Padding(2.0f, 0.0f, 0.0f, 0.0f)
 				[				
 					SNew(SHyperlink)
-						.Text(FText::FromString(Item->PluginStatus.Descriptor.CreatedBy))
+						.Text(FText::FromString(PluginDescriptor.CreatedBy))
 						.ToolTipText(FText::Format(LOCTEXT("NavigateToCreatedByURL", "Visit the vendor's web site ({0})"), FText::FromString(CreatedByURL)))
 						.OnNavigate_Lambda([=]() { FPlatformProcess::LaunchURL(*CreatedByURL, nullptr, nullptr); })
 				];
@@ -171,8 +178,8 @@ void SPluginListTile::Construct( const FArguments& Args, const TSharedRef<SPlugi
 										.Padding(PaddingAmount)
 										[
 											SNew(STextBlock)
-												.Text(FText::FromString(Item->PluginStatus.Descriptor.FriendlyName))
-												.HighlightText_Raw(&Owner->GetOwner().GetPluginTextFilter(), &FPluginTextFilter::GetRawFilterText)
+												.Text(FText::FromString(PluginDescriptor.FriendlyName))
+												.HighlightText_Raw(&OwnerWeak.Pin()->GetOwner().GetPluginTextFilter(), &FPluginTextFilter::GetRawFilterText)
 												.TextStyle(FPluginStyle::Get(), "PluginTile.NameText")
 										]
 
@@ -197,7 +204,7 @@ void SPluginListTile::Construct( const FArguments& Args, const TSharedRef<SPlugi
 															.Padding(0.0f, 0.0f, 3.0f, 2.0f)
 															[
 																SNew(SImage)
-																	.Visibility(Item->PluginStatus.Descriptor.bIsBetaVersion ? EVisibility::Visible : EVisibility::Collapsed)
+																	.Visibility(PluginDescriptor.bIsBetaVersion ? EVisibility::Visible : EVisibility::Collapsed)
 																	.Image(FPluginStyle::Get()->GetBrush("PluginTile.BetaWarning"))
 															]
 
@@ -207,7 +214,7 @@ void SPluginListTile::Construct( const FArguments& Args, const TSharedRef<SPlugi
 															.VAlign(VAlign_Bottom)
 															[
 																SNew(STextBlock)
-																	.Text(Item->PluginStatus.Descriptor.bIsBetaVersion ? LOCTEXT("PluginBetaVersionLabel", "BETA Version ") : LOCTEXT("PluginVersionLabel", "Version "))
+																	.Text(PluginDescriptor.bIsBetaVersion ? LOCTEXT("PluginBetaVersionLabel", "BETA Version ") : LOCTEXT("PluginVersionLabel", "Version "))
 															]
 													]
 
@@ -217,7 +224,7 @@ void SPluginListTile::Construct( const FArguments& Args, const TSharedRef<SPlugi
 													.Padding( 0.0f, 0.0f, 2.0f, 0.0f )	// Extra padding from the right edge
 													[
 														SNew(STextBlock)
-															.Text(FText::FromString(Item->PluginStatus.Descriptor.VersionName))
+															.Text(FText::FromString(PluginDescriptor.VersionName))
 															.TextStyle(FPluginStyle::Get(), "PluginTile.VersionNumberText")
 													]
 											]
@@ -232,7 +239,7 @@ void SPluginListTile::Construct( const FArguments& Args, const TSharedRef<SPlugi
 											.Padding( PaddingAmount )
 											[
 												SNew(STextBlock)
-													.Text(FText::FromString(Item->PluginStatus.Descriptor.Description))
+													.Text(FText::FromString(PluginDescriptor.Description))
 													.AutoWrapText(true)
 											]
 
@@ -245,7 +252,7 @@ void SPluginListTile::Construct( const FArguments& Args, const TSharedRef<SPlugi
 												// Enable checkbox
 												+ SHorizontalBox::Slot()
 													.Padding(PaddingAmount)
-													.AutoWidth()
+													.HAlign(HAlign_Left)
 													[
 														SNew(SCheckBox)
 															.OnCheckStateChanged(this, &SPluginListTile::OnEnablePluginCheckboxChanged)
@@ -258,6 +265,25 @@ void SPluginListTile::Construct( const FArguments& Args, const TSharedRef<SPlugi
 															]
 													]
 
+												// edit link
+												+ SHorizontalBox::Slot()
+													.HAlign(HAlign_Center)
+													.AutoWidth()
+													.Padding(2.0f, 0.0f, 0.0f, 0.0f)
+													[
+														SNew(SHorizontalBox)
+
+														+ SHorizontalBox::Slot()
+															.AutoWidth()
+															.Padding(PaddingAmount)
+															[
+																SNew(SHyperlink)
+																	.Visibility(this, &SPluginListTile::GetAuthoringButtonsVisibility)	
+																	.OnNavigate(this, &SPluginListTile::OnEditPlugin)
+																	.Text(LOCTEXT("EditPlugin", "Edit..."))
+															]
+													]
+
 												// docs link
 												+ SHorizontalBox::Slot()
 													.Padding(PaddingAmount)
@@ -285,21 +311,27 @@ void SPluginListTile::Construct( const FArguments& Args, const TSharedRef<SPlugi
 
 ECheckBoxState SPluginListTile::IsPluginEnabled() const
 {
-	return ItemData->PluginStatus.bIsEnabled
-		? ECheckBoxState::Checked
-		: ECheckBoxState::Unchecked;
+	FPluginBrowserModule& PluginBrowserModule = FPluginBrowserModule::Get();
+	if(PluginBrowserModule.PendingEnablePlugins.Contains(ItemData->Plugin))
+	{
+		return PluginBrowserModule.PendingEnablePlugins[ItemData->Plugin]? ECheckBoxState::Checked : ECheckBoxState::Unchecked;;
+	}
+	else
+	{
+		return ItemData->Plugin->IsEnabled()? ECheckBoxState::Checked : ECheckBoxState::Unchecked;
+	}
 }
 
-
 void SPluginListTile::OnEnablePluginCheckboxChanged(ECheckBoxState NewCheckedState)
 {
 	const bool bNewEnabledState = (NewCheckedState == ECheckBoxState::Checked);
 
-	if (bNewEnabledState && ItemData->PluginStatus.Descriptor.bIsBetaVersion)
+	const FPluginDescriptor& PluginDescriptor = ItemData->Plugin->GetDescriptor();
+	if (bNewEnabledState && PluginDescriptor.bIsBetaVersion)
 	{
 		FText WarningMessage = FText::Format(
 			LOCTEXT("Warning_EnablingBetaPlugin", "Plugin '{0}' is a beta version and might be unstable or removed without notice. Please use with caution. Are you sure you want to enable the plugin?"),
-			FText::FromString(ItemData->PluginStatus.Descriptor.FriendlyName));
+			FText::FromString(PluginDescriptor.FriendlyName));
 
 		if (EAppReturnType::No == FMessageDialog::Open(EAppMsgType::YesNo, WarningMessage))
 		{
@@ -308,14 +340,138 @@ void SPluginListTile::OnEnablePluginCheckboxChanged(ECheckBoxState NewCheckedSta
 		}
 	}
 
-	ItemData->PluginStatus.bIsEnabled = bNewEnabledState;
 	FGameProjectGenerationModule::Get().TryMakeProjectFileWriteable(FPaths::GetProjectFilePath());
 	FText FailMessage;
 
-	if (!IProjectManager::Get().SetPluginEnabled(ItemData->PluginStatus.Name, bNewEnabledState, FailMessage))
+	if (!IProjectManager::Get().SetPluginEnabled(ItemData->Plugin->GetName(), bNewEnabledState, FailMessage))
 	{
 		FMessageDialog::Open(EAppMsgType::Ok, FailMessage);
 	}
+	else
+	{
+		FPluginBrowserModule& PluginBrowserModule = FPluginBrowserModule::Get();
+		if(ItemData->Plugin->IsEnabled() == bNewEnabledState)
+		{
+			PluginBrowserModule.PendingEnablePlugins.Remove(ItemData->Plugin);
+		}
+		else
+		{
+			PluginBrowserModule.PendingEnablePlugins.FindOrAdd(ItemData->Plugin) = bNewEnabledState;
+		}
+	}
+
+
+}
+
+EVisibility SPluginListTile::GetAuthoringButtonsVisibility() const
+{
+	return (FApp::IsEngineInstalled() && ItemData->Plugin->GetLoadedFrom() == EPluginLoadedFrom::Engine)? EVisibility::Hidden : EVisibility::Visible;
+}
+
+void SPluginListTile::OnEditPlugin()
+{
+	// Construct the plugin metadata object using the descriptor for this plugin
+	UPluginMetadataObject* MetadataObject = NewObject<UPluginMetadataObject>();
+	MetadataObject->TargetIconPath = ItemData->Plugin->GetBaseDir() / TEXT("Resources/Icon128.png");
+	MetadataObject->PopulateFromDescriptor(ItemData->Plugin->GetDescriptor());
+	MetadataObject->AddToRoot();
+
+	// Create a property view
+	FPropertyEditorModule& EditModule = FModuleManager::Get().GetModuleChecked<FPropertyEditorModule>("PropertyEditor");
+	TSharedRef<IDetailsView> PropertyView = EditModule.CreateDetailView(FDetailsViewArgs(false, false, false, FDetailsViewArgs::ActorsUseNameArea, true));
+	PropertyView->SetObject(MetadataObject, true);
+
+	// Create the window
+	PropertiesWindow = SNew(SWindow)
+		.SupportsMaximize(false)
+		.SupportsMinimize(false)
+		.SizingRule(ESizingRule::FixedSize)
+		.ClientSize(FVector2D(700.0f, 575.0f))
+		.Title(LOCTEXT("PluginMetadata", "Plugin Properties"))
+		.Content()
+		[
+			SNew(SBorder)
+			.Padding(FMargin(8.0f, 8.0f))
+			[
+				SNew(SVerticalBox)
+
+				+SVerticalBox::Slot()
+				.AutoHeight()
+				.Padding(FMargin(5.0f, 10.0f, 5.0f, 5.0f))
+				[
+					SNew(STextBlock)
+					.Font(FPluginStyle::Get()->GetFontStyle(TEXT("PluginMetadataNameFont")))
+					.Text(FText::FromString(ItemData->Plugin->GetName()))
+				]
+
+				+ SVerticalBox::Slot()
+				.Padding(5)
+				[
+					PropertyView
+				]
+
+				+ SVerticalBox::Slot()
+				.AutoHeight()
+				.Padding(5)
+				.HAlign(HAlign_Right)
+				[
+					SNew(SButton)
+					.ContentPadding(FMargin(20.0f, 2.0f))
+					.Text(LOCTEXT("OkButtonLabel", "Ok"))
+					.OnClicked(this, &SPluginListTile::OnEditPluginFinished, MetadataObject)
+				]
+			]
+		];
+
+	FSlateApplication::Get().AddModalWindow(PropertiesWindow.ToSharedRef(), OwnerWeak.Pin());//Args.ParentWidget);
+}
+
+FReply SPluginListTile::OnEditPluginFinished(UPluginMetadataObject* MetadataObject)
+{
+	IPlugin* Plugin = ItemData->Plugin;
+	FPluginDescriptor OldDescriptor = Plugin->GetDescriptor();
+
+	// Update the descriptor with the new metadata
+	FPluginDescriptor NewDescriptor = OldDescriptor;
+	MetadataObject->CopyIntoDescriptor(NewDescriptor);
+	MetadataObject->RemoveFromRoot();
+
+	// Close the properties window
+	PropertiesWindow->RequestDestroyWindow();
+
+	// Write both to strings
+	FString OldText = OldDescriptor.ToString();
+	FString NewText = NewDescriptor.ToString();
+	if(OldText.Compare(NewText, ESearchCase::CaseSensitive) != 0)
+	{
+		FString DescriptorFileName = Plugin->GetDescriptorFileName();
+
+		// First attempt to check out the file if SCC is enabled
+		ISourceControlModule& SourceControlModule = ISourceControlModule::Get();
+		if(SourceControlModule.IsEnabled())
+		{
+			ISourceControlProvider& SourceControlProvider = SourceControlModule.GetProvider();
+			TSharedPtr<ISourceControlState, ESPMode::ThreadSafe> SourceControlState = SourceControlProvider.GetState(DescriptorFileName, EStateCacheUsage::ForceUpdate);
+			if(SourceControlState.IsValid() && SourceControlState->CanCheckout())
+			{
+				SourceControlProvider.Execute(ISourceControlOperation::Create<FCheckOut>(), DescriptorFileName);
+			}
+		}
+
+		// Write to the file and update the in-memory metadata
+		FText FailReason;
+		if(!ItemData->Plugin->UpdateDescriptor(NewDescriptor, FailReason))
+		{
+			FMessageDialog::Open(EAppMsgType::Ok, FailReason);
+		}
+
+		// Refresh the view
+		if(OwnerWeak.IsValid())
+		{
+			OwnerWeak.Pin()->GetOwner().SetNeedsRefresh();
+		}
+	}
+	return FReply::Handled();
 }
 
 #undef LOCTEXT_NAMESPACE
diff --git a/Engine/Plugins/Editor/PluginBrowser/Source/PluginBrowser/Private/SPluginListTile.h b/Engine/Plugins/Editor/PluginBrowser/Source/PluginBrowser/Private/SPluginListTile.h
--- a/Engine/Plugins/Editor/PluginBrowser/Source/PluginBrowser/Private/SPluginListTile.h
+++ b/Engine/Plugins/Editor/PluginBrowser/Source/PluginBrowser/Private/SPluginListTile.h
@@ -26,6 +26,14 @@ private:
 	/** Called when the enabled checkbox is clicked */
 	void OnEnablePluginCheckboxChanged(ECheckBoxState NewCheckedState);
 
+	/** Used to determine whether to show the edit and package buttons for this plugin */
+	EVisibility GetAuthoringButtonsVisibility() const;
+
+	/** Called when the 'edit' hyperlink is clicked */
+	void OnEditPlugin();
+
+	/** Called when the 'edit' hyperlink is clicked */
+	FReply OnEditPluginFinished(UPluginMetadataObject* MetadataObject);
 
 private:
 
@@ -35,6 +43,9 @@ private:
 	/** Weak pointer back to its owner */
 	TWeakPtr< class SPluginList > OwnerWeak;
 
+	/** Dialog for editing plugin properties */
+	TSharedPtr< SWindow > PropertiesWindow;
+
 	/** Brush resource for the image that is dynamically loaded */
 	TSharedPtr< FSlateDynamicImageBrush > PluginIconDynamicImageBrush;
 };
diff --git a/Engine/Plugins/Editor/PluginBrowser/Source/PluginBrowser/Private/SPluginsEditor.cpp b/Engine/Plugins/Editor/PluginBrowser/Source/PluginBrowser/Private/SPluginsEditor.cpp
--- a/Engine/Plugins/Editor/PluginBrowser/Source/PluginBrowser/Private/SPluginsEditor.cpp
+++ b/Engine/Plugins/Editor/PluginBrowser/Source/PluginBrowser/Private/SPluginsEditor.cpp
@@ -15,21 +15,21 @@
 
 void SPluginsEditor::Construct( const FArguments& Args )
 {
+	// @todo plugedit: Should we save/restore category selection and other view settings?  Splitter position, etc.
+
+	RegisterActiveTimer (0.f, FWidgetActiveTimerDelegate::CreateSP (this, &SPluginsEditor::TriggerBreadcrumbRefresh));
+
 	struct Local
 	{
-		static void PluginStatusToStringArray( const FPluginStatus& PluginStatus, OUT TArray< FString >& StringArray )
+		static void PluginToStringArray( const IPlugin* Plugin, OUT TArray< FString >& StringArray )
 		{
 			// NOTE: Only the friendly name is searchable for now.  We don't display the actual plugin name in the UI.
-			StringArray.Add( PluginStatus.Descriptor.FriendlyName );
+			StringArray.Add( Plugin->GetDescriptor().FriendlyName );
 		}
 	};
 
-	// @todo plugedit: Should we save/restore category selection and other view settings?  Splitter position, etc.
-
-	RegisterActiveTimer (0.f, FWidgetActiveTimerDelegate::CreateSP (this, &SPluginsEditor::TriggerBreadcrumbRefresh));
-
 	// Setup text filtering
-	PluginTextFilter = MakeShareable( new FPluginTextFilter( FPluginTextFilter::FItemToStringArray::CreateStatic( &Local::PluginStatusToStringArray ) ) );
+	PluginTextFilter = MakeShareable( new FPluginTextFilter( FPluginTextFilter::FItemToStringArray::CreateStatic( &Local::PluginToStringArray ) ) );
 
 	const float PaddingAmount = 2.0f;
 
@@ -162,6 +162,22 @@ void SPluginsEditor::OnCategorySelectionChanged()
 	RegisterActiveTimer (0.f, FWidgetActiveTimerDelegate::CreateSP (this, &SPluginsEditor::TriggerBreadcrumbRefresh));
 }
 
+void SPluginsEditor::SetNeedsRefresh()
+{
+	if( PluginList.IsValid() )
+	{
+		PluginList->SetNeedsRefresh();
+	}
+
+	if( PluginCategories.IsValid() )
+	{
+		PluginCategories->SetNeedsRefresh();
+	}
+
+	// Breadcrumbs will need to be refreshed
+	RegisterActiveTimer (0.f, FWidgetActiveTimerDelegate::CreateSP (this, &SPluginsEditor::TriggerBreadcrumbRefresh));
+}
+
 EActiveTimerReturnType SPluginsEditor::TriggerBreadcrumbRefresh(double InCurrentTime, float InDeltaTime)
 {
 	RefreshBreadcrumbTrail();
diff --git a/Engine/Plugins/Editor/PluginBrowser/Source/PluginBrowser/Private/SPluginsEditor.h b/Engine/Plugins/Editor/PluginBrowser/Source/PluginBrowser/Private/SPluginsEditor.h
--- a/Engine/Plugins/Editor/PluginBrowser/Source/PluginBrowser/Private/SPluginsEditor.h
+++ b/Engine/Plugins/Editor/PluginBrowser/Source/PluginBrowser/Private/SPluginsEditor.h
@@ -64,6 +64,8 @@ class SPluginsEditor : public SCompoundWidget
 	/** Called when the selected category changes so we can invalidate the list */
 	void OnCategorySelectionChanged();
 
+	/** Refresh the whole window */
+	void SetNeedsRefresh();
 
 private:
 
@@ -98,6 +100,5 @@ private:
 
 	/** Breadcrumb trail widget for the currently selected category */
 	TSharedPtr< SPluginCategoryBreadcrumbTrail > BreadcrumbTrail;
-
 };
 
diff --git a/Engine/Source/Runtime/Projects/Private/PluginDescriptor.cpp b/Engine/Source/Runtime/Projects/Private/PluginDescriptor.cpp
--- a/Engine/Source/Runtime/Projects/Private/PluginDescriptor.cpp
+++ b/Engine/Source/Runtime/Projects/Private/PluginDescriptor.cpp
@@ -26,23 +26,23 @@ bool FPluginDescriptor::Load( const FString& FileName, FText& OutFailReason )
 		return false;
 	}
 
+	// Parse it as a plug-in descriptor
+	return Read(FileContents, OutFailReason);
+}
+
+
+bool FPluginDescriptor::Read(const FString& Text, FText& OutFailReason)
+{
 	// Deserialize a JSON object from the string
-	TSharedPtr< FJsonObject > Object;
-	TSharedRef< TJsonReader<> > Reader = TJsonReaderFactory<>::Create(FileContents);
-	
-	if (!FJsonSerializer::Deserialize(Reader, Object) || !Object.IsValid() )
+	TSharedPtr< FJsonObject > ObjectPtr;
+	TSharedRef< TJsonReader<> > Reader = TJsonReaderFactory<>::Create(Text);
+	if (!FJsonSerializer::Deserialize(Reader, ObjectPtr) || !ObjectPtr.IsValid() )
 	{
 		OutFailReason = FText::Format(LOCTEXT("FailedToReadDescriptorFile", "Failed to read file. {0}"), FText::FromString(Reader->GetErrorMessage()));
 		return false;
 	}
+	FJsonObject& Object = *ObjectPtr.Get();
 
-	// Parse it as a plug-in descriptor
-	return Read(*Object.Get(), OutFailReason);
-}
-
-
-bool FPluginDescriptor::Read( const FJsonObject& Object, FText& OutFailReason )
-{
 	// Read the file version
 	int32 FileVersionInt32;
 
@@ -100,6 +100,52 @@ bool FPluginDescriptor::Read( const FJsonObject& Object, FText& OutFailReason )
 	return true;
 }
 
+bool FPluginDescriptor::Save(const FString& FileName, FText& OutFailReason) const
+{
+	// Write the contents of the descriptor to a string. Make sure the writer is destroyed so that the contents are flushed to the string.
+	FString Text = ToString();
+	if ( FFileHelper::SaveStringToFile(Text, *FileName) )
+	{
+		return true;
+	}
+	else
+	{
+		OutFailReason = FText::Format( LOCTEXT("FailedToWriteOutputFile", "Failed to write output file '{0}'. Perhaps the file is Read-Only?"), FText::FromString(FileName) );
+		return false;
+	}
+}
+
+FString FPluginDescriptor::ToString() const
+{
+	FString Text;
+	TSharedRef< TJsonWriter<> > WriterRef = TJsonWriterFactory<>::Create(&Text);
+	TJsonWriter<>& Writer = WriterRef.Get();
+
+	Writer.WriteObjectStart();
+
+	Writer.WriteValue(TEXT("FileVersion"), EProjectDescriptorVersion::Latest);
+	Writer.WriteValue(TEXT("Version"), Version);
+	Writer.WriteValue(TEXT("VersionName"), VersionName);
+	Writer.WriteValue(TEXT("FriendlyName"), FriendlyName);
+	Writer.WriteValue(TEXT("Description"), Description);
+	Writer.WriteValue(TEXT("Category"), Category);
+	Writer.WriteValue(TEXT("CreatedBy"), CreatedBy);
+	Writer.WriteValue(TEXT("CreatedByURL"), CreatedByURL);
+	Writer.WriteValue(TEXT("DocsURL"), DocsURL);
+
+	FModuleDescriptor::WriteArray(Writer, TEXT("Modules"), Modules);
+
+	Writer.WriteValue(TEXT("EnabledByDefault"), bEnabledByDefault);
+	Writer.WriteValue(TEXT("CanContainContent"), bCanContainContent);
+	Writer.WriteValue(TEXT("IsBetaVersion"), bIsBetaVersion);
+
+	Writer.WriteObjectEnd();
+	Writer.Close();
+
+	return Text;
+}
+
+//////////////////////////////////////////////////////////////////////////////////////////////////////////////////
 
 FPluginReferenceDescriptor::FPluginReferenceDescriptor( const FString &InName, bool bInEnabled )
 	: Name(InName)
diff --git a/Engine/Source/Runtime/Projects/Private/PluginManager.cpp b/Engine/Source/Runtime/Projects/Private/PluginManager.cpp
--- a/Engine/Source/Runtime/Projects/Private/PluginManager.cpp
+++ b/Engine/Source/Runtime/Projects/Private/PluginManager.cpp
@@ -33,6 +33,11 @@ FString FPlugin::GetName() const
 	return Name;
 }
 
+FString FPlugin::GetDescriptorFileName() const
+{
+	return FileName;
+}
+
 FString FPlugin::GetBaseDir() const
 {
 	return FPaths::GetPath(FileName);
@@ -58,6 +63,26 @@ bool FPlugin::CanContainContent() const
 	return Descriptor.bCanContainContent;
 }
 
+EPluginLoadedFrom FPlugin::GetLoadedFrom() const
+{
+	return LoadedFrom;
+}
+
+const FPluginDescriptor& FPlugin::GetDescriptor() const
+{
+	return Descriptor;
+}
+
+bool FPlugin::UpdateDescriptor(const FPluginDescriptor& NewDescriptor, FText& OutFailReason)
+{
+	if(!NewDescriptor.Save(FileName, OutFailReason))
+	{
+		return false;
+	}
+
+	Descriptor = NewDescriptor;
+	return true;
+}
 
 
 
@@ -562,6 +587,16 @@ TArray<IPlugin*> FPluginManager::GetEnabledPlugins()
 	return Plugins;
 }
 
+TArray<IPlugin*> FPluginManager::GetDiscoveredPlugins()
+{
+	TArray<IPlugin*> Plugins;
+	for(TSharedRef<FPlugin>& Plugin : AllPlugins)
+	{
+		Plugins.Add(&(Plugin.Get()));
+	}
+	return Plugins;
+}
+
 TArray< FPluginStatus > FPluginManager::QueryStatusForAllPlugins() const
 {
 	TArray< FPluginStatus > PluginStatuses;
diff --git a/Engine/Source/Runtime/Projects/Private/PluginManager.h b/Engine/Source/Runtime/Projects/Private/PluginManager.h
--- a/Engine/Source/Runtime/Projects/Private/PluginManager.h
+++ b/Engine/Source/Runtime/Projects/Private/PluginManager.h
@@ -37,11 +37,15 @@ public:
 
 	/* IPluginInfo interface */
 	virtual FString GetName() const override;
+	virtual FString GetDescriptorFileName() const override;
 	virtual FString GetBaseDir() const override;
 	virtual FString GetContentDir() const override;
 	virtual FString GetMountedAssetPath() const override;
 	virtual bool IsEnabled() const override;
 	virtual bool CanContainContent() const override;
+	virtual EPluginLoadedFrom GetLoadedFrom() const override;
+	virtual const FPluginDescriptor& GetDescriptor() const override;
+	virtual bool UpdateDescriptor(const FPluginDescriptor& NewDescriptor, FText& OutFailReason) override;
 };
 
 /**
@@ -63,6 +67,7 @@ public:
 	virtual bool CheckModuleCompatibility( TArray<FString>& OutIncompatibleModules ) override;
 	virtual FPlugin* FindPlugin(const FString& Name) override;
 	virtual TArray<IPlugin*> GetEnabledPlugins() override;
+	virtual TArray<IPlugin*> GetDiscoveredPlugins() override;
 	virtual TArray< FPluginStatus > QueryStatusForAllPlugins() const override;
 
 private:
diff --git a/Engine/Source/Runtime/Projects/Public/Interfaces/IPluginManager.h b/Engine/Source/Runtime/Projects/Public/Interfaces/IPluginManager.h
--- a/Engine/Source/Runtime/Projects/Public/Interfaces/IPluginManager.h
+++ b/Engine/Source/Runtime/Projects/Public/Interfaces/IPluginManager.h
@@ -53,6 +53,13 @@ public:
 	 */
 	virtual FString GetName() const = 0;
 
+	/**
+	 * Get a path to the plugin's descriptor
+	 *
+	 * @return Path to the plugin's descriptor.
+	 */
+	virtual FString GetDescriptorFileName() const = 0;
+
 	/**
 	 * Get a path to the plugin's directory.
 	 *
@@ -87,6 +94,29 @@ public:
 	 * @return True if the plugin can contain content.
 	 */
 	virtual bool CanContainContent() const = 0;
+
+	/**
+	 * Returns the plugin's location
+	 *
+	 * @return Where the plugin was loaded from
+	 */
+	virtual EPluginLoadedFrom GetLoadedFrom() const = 0;
+
+	/**
+	 * Gets the plugin's descriptor
+	 *
+	 * @return Reference to the plugin's descriptor
+	 */
+	virtual const FPluginDescriptor& GetDescriptor() const = 0;
+
+	/**
+	 * Updates the plugin's descriptor
+	 *
+	 * @param NewDescriptor The new plugin descriptor
+	 * @param OutFailReason The error message if the plugin's descriptor could not be updated
+	 * @return True if the descriptor was updated, false otherwise. 
+	 */ 
+	virtual bool UpdateDescriptor(const FPluginDescriptor& NewDescriptor, FText& OutFailReason) = 0;
 };
 
 /**
@@ -144,6 +174,13 @@ public:
 	 */
 	virtual TArray<IPlugin*> GetEnabledPlugins() = 0;
 
+	/**
+	 * Gets an array of all the discovered plugins.
+	 *
+	 * @return	Array of the discovered plugins.
+	 */
+	virtual TArray<IPlugin*> GetDiscoveredPlugins() = 0;
+
 	/**
 	 * Gets status about all currently known plug-ins.
 	 *
diff --git a/Engine/Source/Runtime/Projects/Public/PluginDescriptor.h b/Engine/Source/Runtime/Projects/Public/PluginDescriptor.h
--- a/Engine/Source/Runtime/Projects/Public/PluginDescriptor.h
+++ b/Engine/Source/Runtime/Projects/Public/PluginDescriptor.h
@@ -77,7 +77,13 @@ struct PROJECTS_API FPluginDescriptor
 	bool Load(const FString& FileName, FText& OutFailReason);
 
 	/** Reads the descriptor from the given JSON object */
-	bool Read(const FJsonObject& Object, FText& OutFailReason);
+	bool Read(const FString& Text, FText& OutFailReason);
+
+	/** Saves the descriptor from the given file. */
+	bool Save(const FString& FileName, FText& OutFailReason) const;
+
+	/** Writes a descriptor to JSON */
+	FString ToString() const;
 };
 
 /**
