Renaming some plugin browser stuff for consistency, and removing a bunch of unnecessary abstractions.

diff --git a/Engine/Plugins/Editor/PluginBrowser/Source/PluginBrowser/Private/PluginBrowserModule.cpp b/Engine/Plugins/Editor/PluginBrowser/Source/PluginBrowser/Private/PluginBrowserModule.cpp
--- a/Engine/Plugins/Editor/PluginBrowser/Source/PluginBrowser/Private/PluginBrowserModule.cpp
+++ b/Engine/Plugins/Editor/PluginBrowser/Source/PluginBrowser/Private/PluginBrowserModule.cpp
@@ -1,7 +1,7 @@
 // Copyright 1998-2015 Epic Games, Inc. All Rights Reserved.
 
-#include "PluginsEditorPrivatePCH.h"
-#include "SPluginsEditor.h"
+#include "PluginBrowserPrivatePCH.h"
+#include "SPluginBrowser.h"
 #include "Runtime/Core/Public/Features/IModularFeatures.h"
 #include "Editor/UnrealEd/Public/Features/EditorFeatures.h"
 #include "PluginMetadataObject.h"
@@ -34,7 +34,7 @@ void FPluginBrowserModule::StartupModule()
 				.Icon( FPluginStyle::Get()->GetBrush("Plugins.TabIcon") )
 				.TabRole( ETabRole::MajorTab );
 
-			MajorTab->SetContent( SNew( SPluginsEditor ) );
+			MajorTab->SetContent( SNew( SPluginBrowser ) );
 
 			return MajorTab;
 		}
diff --git a/Engine/Plugins/Editor/PluginBrowser/Source/PluginBrowser/Private/PluginBrowserModule.h b/Engine/Plugins/Editor/PluginBrowser/Source/PluginBrowser/Private/PluginBrowserModule.h
--- a/Engine/Plugins/Editor/PluginBrowser/Source/PluginBrowser/Private/PluginBrowserModule.h
+++ b/Engine/Plugins/Editor/PluginBrowser/Source/PluginBrowser/Private/PluginBrowserModule.h
@@ -2,10 +2,11 @@
 
 #pragma once 
 
+#include "IPluginBrowser.h"
 #include "IPluginManager.h"
 #include "ModuleManager.h"
 
-class FPluginBrowserModule : public IPluginsEditor
+class FPluginBrowserModule : public IPluginBrowser
 {
 public:
 	TMap<IPlugin*, bool> PendingEnablePlugins;
diff --git a/Engine/Plugins/Editor/PluginBrowser/Source/PluginBrowser/Private/PluginBrowserPrivatePCH.h b/Engine/Plugins/Editor/PluginBrowser/Source/PluginBrowser/Private/PluginBrowserPrivatePCH.h
new file mode 100644
--- /dev/null
+++ b/Engine/Plugins/Editor/PluginBrowser/Source/PluginBrowser/Private/PluginBrowserPrivatePCH.h
@@ -0,0 +1,11 @@
+// Copyright 1998-2015 Epic Games, Inc. All Rights Reserved.
+
+#pragma once
+
+#include "UnrealEd.h"
+#include "IPluginBrowser.h"
+#include "IPluginManager.h"
+
+#include "TextFilter.h"
+
+typedef TTextFilter< const IPlugin* > FPluginTextFilter;
diff --git a/Engine/Plugins/Editor/PluginBrowser/Source/PluginBrowser/Private/PluginCategoryTreeItem.h b/Engine/Plugins/Editor/PluginBrowser/Source/PluginBrowser/Private/PluginCategoryTreeItem.h
deleted file mode 100644
diff --git a/Engine/Plugins/Editor/PluginBrowser/Source/PluginBrowser/Private/PluginListItem.h b/Engine/Plugins/Editor/PluginBrowser/Source/PluginBrowser/Private/PluginListItem.h
deleted file mode 100644
diff --git a/Engine/Plugins/Editor/PluginBrowser/Source/PluginBrowser/Private/PluginMetadataObject.cpp b/Engine/Plugins/Editor/PluginBrowser/Source/PluginBrowser/Private/PluginMetadataObject.cpp
--- a/Engine/Plugins/Editor/PluginBrowser/Source/PluginBrowser/Private/PluginMetadataObject.cpp
+++ b/Engine/Plugins/Editor/PluginBrowser/Source/PluginBrowser/Private/PluginMetadataObject.cpp
@@ -1,6 +1,6 @@
 // Copyright 1998-2015 Epic Games, Inc. All Rights Reserved.
 
-#include "PluginsEditorPrivatePCH.h"
+#include "PluginBrowserPrivatePCH.h"
 #include "DetailLayoutBuilder.h"
 #include "DetailCategoryBuilder.h"
 #include "DetailWidgetRow.h"
diff --git a/Engine/Plugins/Editor/PluginBrowser/Source/PluginBrowser/Private/PluginStyle.cpp b/Engine/Plugins/Editor/PluginBrowser/Source/PluginBrowser/Private/PluginStyle.cpp
--- a/Engine/Plugins/Editor/PluginBrowser/Source/PluginBrowser/Private/PluginStyle.cpp
+++ b/Engine/Plugins/Editor/PluginBrowser/Source/PluginBrowser/Private/PluginStyle.cpp
@@ -1,6 +1,6 @@
 // Copyright 1998-2015 Epic Games, Inc. All Rights Reserved.
 
-#include "PluginsEditorPrivatePCH.h"
+#include "PluginBrowserPrivatePCH.h"
 #include "PluginStyle.h"
 #include "SlateStyle.h"
 #include "IPluginManager.h"
diff --git a/Engine/Plugins/Editor/PluginBrowser/Source/PluginBrowser/Private/PluginsEditorPrivatePCH.h b/Engine/Plugins/Editor/PluginBrowser/Source/PluginBrowser/Private/PluginsEditorPrivatePCH.h
deleted file mode 100644
diff --git a/Engine/Plugins/Editor/PluginBrowser/Source/PluginBrowser/Private/SPluginBrowser.cpp b/Engine/Plugins/Editor/PluginBrowser/Source/PluginBrowser/Private/SPluginBrowser.cpp
new file mode 100644
--- /dev/null
+++ b/Engine/Plugins/Editor/PluginBrowser/Source/PluginBrowser/Private/SPluginBrowser.cpp
@@ -0,0 +1,225 @@
+// Copyright 1998-2015 Epic Games, Inc. All Rights Reserved.
+
+#include "PluginBrowserPrivatePCH.h"
+#include "PluginStyle.h"
+#include "SPluginBrowser.h"
+#include "SPluginList.h"
+#include "SPluginCategories.h"
+#include "SSearchBox.h"
+
+#define LOCTEXT_NAMESPACE "PluginsEditor"
+
+
+// @todo plugedit: Intro anim!!!
+
+
+void SPluginBrowser::Construct( const FArguments& Args )
+{
+	// @todo plugedit: Should we save/restore category selection and other view settings?  Splitter position, etc.
+
+	RegisterActiveTimer (0.f, FWidgetActiveTimerDelegate::CreateSP (this, &SPluginBrowser::TriggerBreadcrumbRefresh));
+
+	struct Local
+	{
+		static void PluginToStringArray( const IPlugin* Plugin, OUT TArray< FString >& StringArray )
+		{
+			// NOTE: Only the friendly name is searchable for now.  We don't display the actual plugin name in the UI.
+			StringArray.Add( Plugin->GetDescriptor().FriendlyName );
+		}
+	};
+
+	// Setup text filtering
+	PluginTextFilter = MakeShareable( new FPluginTextFilter( FPluginTextFilter::FItemToStringArray::CreateStatic( &Local::PluginToStringArray ) ) );
+
+	const float PaddingAmount = 2.0f;
+
+
+	PluginCategories = SNew( SPluginCategories, SharedThis( this ) );
+
+
+	ChildSlot
+	[ 
+		SNew( SHorizontalBox )
+
+		+SHorizontalBox::Slot()
+		.Padding( PaddingAmount )
+		.AutoWidth()	// @todo plugedit: Probably want a splitter here (w/ saved layout)
+		[
+			PluginCategories.ToSharedRef()
+		]
+
+		+SHorizontalBox::Slot()
+		.Padding( PaddingAmount )
+		[
+			SNew( SVerticalBox )
+
+			+SVerticalBox::Slot()
+			.Padding(FMargin(0.0f, 0.0f, 0.0f, PaddingAmount))
+			.AutoHeight()
+			[
+				SNew( SHorizontalBox )
+
+				+SHorizontalBox::Slot()
+				.Padding( PaddingAmount )
+				[
+					SAssignNew( BreadcrumbTrail, SPluginCategoryBreadcrumbTrail )
+					.DelimiterImage( FPluginStyle::Get()->GetBrush( "Plugins.BreadcrumbArrow" ) ) 
+					.ShowLeadingDelimiter( true )
+					.OnCrumbClicked( this, &SPluginBrowser::BreadcrumbTrail_OnCrumbClicked )
+				]
+	
+				+SHorizontalBox::Slot()
+				.Padding( PaddingAmount )
+				[
+					SNew( SSearchBox )
+					.OnTextChanged( this, &SPluginBrowser::SearchBox_OnPluginSearchTextChanged )
+				]
+			]
+
+			+SVerticalBox::Slot()
+			[
+				SAssignNew( PluginList, SPluginList, SharedThis( this ) )
+			]
+
+			+SVerticalBox::Slot()
+			.AutoHeight()
+			.Padding(FMargin(PaddingAmount, PaddingAmount, PaddingAmount, 0.0f))
+			[
+				SNew(SBorder)
+				.BorderBackgroundColor(FLinearColor::Yellow)
+				.BorderImage(FEditorStyle::GetBrush("ToolPanel.GroupBorder"))
+				.Padding(8.0f)
+				.Visibility(this, &SPluginBrowser::HandleRestartEditorNoticeVisibility)
+				[
+					SNew( SHorizontalBox )
+
+					+SHorizontalBox::Slot()
+					.AutoWidth()
+					.VAlign(VAlign_Center)
+					.Padding(FMargin(0.0f, 0.0f, 8.0f, 0.0f))
+					[
+						SNew(SImage)
+						.Image(FPluginStyle::Get()->GetBrush("Plugins.Warning"))
+					]
+
+					+SHorizontalBox::Slot()
+					.VAlign(VAlign_Center)
+					[
+						SNew(STextBlock)
+						.Text(LOCTEXT("PluginSettingsRestartNotice", "Unreal Editor must be restarted for the plugin changes to take effect."))
+					]
+
+					+SHorizontalBox::Slot()
+					.AutoWidth()
+					.VAlign(VAlign_Center)
+					.HAlign(HAlign_Right)
+					[
+						SNew(SButton)
+						.Text(LOCTEXT("PluginSettingsRestartEditor", "Restart Now"))
+						.OnClicked(this, &SPluginBrowser::HandleRestartEditorButtonClicked)
+					]
+				]
+			]
+		]
+	];
+}
+
+
+EVisibility SPluginBrowser::HandleRestartEditorNoticeVisibility() const
+{
+	return (FPluginBrowserModule::Get().PendingEnablePlugins.Num() > 0)? EVisibility::Visible : EVisibility::Collapsed;
+}
+
+
+FReply SPluginBrowser::HandleRestartEditorButtonClicked() const
+{
+	const bool bWarn = false;
+	FUnrealEdMisc::Get().RestartEditor(bWarn);
+	return FReply::Handled();
+}
+
+
+void SPluginBrowser::SearchBox_OnPluginSearchTextChanged( const FText& NewText )
+{
+	PluginTextFilter->SetRawFilterText( NewText );
+}
+
+
+TSharedPtr< FPluginCategory > SPluginBrowser::GetSelectedCategory() const
+{
+	return PluginCategories->GetSelectedCategory();
+}
+
+
+void SPluginBrowser::OnCategorySelectionChanged()
+{
+	if( PluginList.IsValid() )
+	{
+		PluginList->SetNeedsRefresh();
+	}
+
+	// Breadcrumbs will need to be refreshed
+	RegisterActiveTimer (0.f, FWidgetActiveTimerDelegate::CreateSP (this, &SPluginBrowser::TriggerBreadcrumbRefresh));
+}
+
+void SPluginBrowser::SetNeedsRefresh()
+{
+	if( PluginList.IsValid() )
+	{
+		PluginList->SetNeedsRefresh();
+	}
+
+	if( PluginCategories.IsValid() )
+	{
+		PluginCategories->SetNeedsRefresh();
+	}
+
+	// Breadcrumbs will need to be refreshed
+	RegisterActiveTimer (0.f, FWidgetActiveTimerDelegate::CreateSP (this, &SPluginBrowser::TriggerBreadcrumbRefresh));
+}
+
+EActiveTimerReturnType SPluginBrowser::TriggerBreadcrumbRefresh(double InCurrentTime, float InDeltaTime)
+{
+	RefreshBreadcrumbTrail();
+	return EActiveTimerReturnType::Stop;
+}
+
+void SPluginBrowser::RefreshBreadcrumbTrail()
+{
+	// Update breadcrumb trail
+	if( BreadcrumbTrail.IsValid() )
+	{
+		TSharedPtr<FPluginCategory> SelectedCategory = PluginCategories->GetSelectedCategory();
+
+		// Build up the list of categories, starting at the selected node and working our way backwards.
+		TArray<TSharedPtr<FPluginCategory>> CategoryPath;
+		if(SelectedCategory.IsValid())
+		{
+			for(TSharedPtr<FPluginCategory> NextCategory = SelectedCategory; NextCategory.IsValid(); NextCategory = NextCategory->ParentCategory.Pin())
+			{
+				CategoryPath.Insert( NextCategory, 0 );
+			}
+		}
+
+		// Fill in the crumbs
+		BreadcrumbTrail->ClearCrumbs();
+		for(TSharedPtr<FPluginCategory>& Category: CategoryPath)
+		{
+			FPluginCategoryBreadcrumbPtr NewBreadcrumb( new FPluginCategoryBreadcrumb( Category ) );
+			BreadcrumbTrail->PushCrumb( Category->DisplayName, NewBreadcrumb );
+		}
+	}
+}
+
+
+void SPluginBrowser::BreadcrumbTrail_OnCrumbClicked( const FPluginCategoryBreadcrumbPtr& CrumbData )
+{
+	if( PluginCategories.IsValid() )
+	{
+		PluginCategories->SelectCategory( CrumbData->GetCategoryItem() );
+	}
+}
+
+
+
+#undef LOCTEXT_NAMESPACE
\ No newline at end of file
diff --git a/Engine/Plugins/Editor/PluginBrowser/Source/PluginBrowser/Private/SPluginBrowser.h b/Engine/Plugins/Editor/PluginBrowser/Source/PluginBrowser/Private/SPluginBrowser.h
new file mode 100644
--- /dev/null
+++ b/Engine/Plugins/Editor/PluginBrowser/Source/PluginBrowser/Private/SPluginBrowser.h
@@ -0,0 +1,104 @@
+// Copyright 1998-2015 Epic Games, Inc. All Rights Reserved.
+
+#pragma once
+
+#include "SBreadcrumbTrail.h"
+
+/**
+ * Stores data about a breadcrumb entry for the plugins editor
+ */
+class FPluginCategoryBreadcrumb
+{
+
+public:
+
+	/** Initializing constructor */
+	FPluginCategoryBreadcrumb( TSharedPtr< class FPluginCategory > InitCategoryItem )
+		: CategoryItem( InitCategoryItem )
+	{
+	}
+
+	/** @return Gets the category item for this breadcrumb */
+	const TSharedPtr< class FPluginCategory >& GetCategoryItem() const
+	{
+		return CategoryItem;
+	}
+
+
+private:
+
+	/** The category tree item data for this breadcrumb */
+	TSharedPtr< class FPluginCategory > CategoryItem;
+};
+
+
+
+typedef TSharedPtr< FPluginCategoryBreadcrumb > FPluginCategoryBreadcrumbPtr;
+typedef SBreadcrumbTrail< FPluginCategoryBreadcrumbPtr > SPluginCategoryBreadcrumbTrail;
+
+
+/**
+ * Implementation of main plugin editor Slate widget
+ */
+class SPluginBrowser : public SCompoundWidget
+{
+	SLATE_BEGIN_ARGS( SPluginBrowser )
+	{
+	}
+
+	SLATE_END_ARGS()
+
+
+	/** Widget constructor */
+	void Construct( const FArguments& Args );
+
+	/** @return Returns the plugin text filter object, so that child widgets can subscribe to find out about changes */
+	FPluginTextFilter& GetPluginTextFilter()
+	{
+		return *PluginTextFilter;
+	}
+
+	/** @return Returns the currently selected category */
+	TSharedPtr< class FPluginCategory > GetSelectedCategory() const;
+
+	/** Called when the selected category changes so we can invalidate the list */
+	void OnCategorySelectionChanged();
+
+	/** Refresh the whole window */
+	void SetNeedsRefresh();
+
+private:
+
+	/** @return Is the "restart required" notice visible? */
+	EVisibility HandleRestartEditorNoticeVisibility() const;
+
+	/** Handle the "restart now" button being clicked */
+	FReply HandleRestartEditorButtonClicked() const;
+
+	/** Called when the text in the search box was changed */
+	void SearchBox_OnPluginSearchTextChanged( const FText& NewText );
+
+	/** Called when a breadcrumb is clicked on the breadcrumb trail */
+	void BreadcrumbTrail_OnCrumbClicked( const FPluginCategoryBreadcrumbPtr& CrumbData );
+
+	/** Called to refresh the breadcrumb trail immediately */
+	void RefreshBreadcrumbTrail();
+
+	/** One-off active timer to trigger a refresh of the breadcrumb trail as needed */
+	EActiveTimerReturnType TriggerBreadcrumbRefresh(double InCurrentTime, float InDeltaTime);
+
+private:
+
+	/** The plugin categories widget */
+	TSharedPtr< class SPluginCategories > PluginCategories;
+
+	/** The plugin list widget */
+	TSharedPtr< class SPluginList > PluginList;
+
+	/** Text filter object for typing in filter text to the search box */
+	TSharedPtr< FPluginTextFilter > PluginTextFilter;
+
+	/** Breadcrumb trail widget for the currently selected category */
+	TSharedPtr< SPluginCategoryBreadcrumbTrail > BreadcrumbTrail;
+};
+
diff --git a/Engine/Plugins/Editor/PluginBrowser/Source/PluginBrowser/Private/SPluginCategories.cpp b/Engine/Plugins/Editor/PluginBrowser/Source/PluginBrowser/Private/SPluginCategories.cpp
--- a/Engine/Plugins/Editor/PluginBrowser/Source/PluginBrowser/Private/SPluginCategories.cpp
+++ b/Engine/Plugins/Editor/PluginBrowser/Source/PluginBrowser/Private/SPluginCategories.cpp
@@ -1,30 +1,29 @@
 // Copyright 1998-2015 Epic Games, Inc. All Rights Reserved.
 
-#include "PluginsEditorPrivatePCH.h"
+#include "PluginBrowserPrivatePCH.h"
 #include "SPluginCategories.h"
-#include "PluginCategoryTreeItem.h"
 #include "IPluginManager.h"
 #include "SPluginCategoryTreeItem.h"
-#include "SPluginsEditor.h"
+#include "SPluginBrowser.h"
 
 #define LOCTEXT_NAMESPACE "PluginCategories"
 
 
 
-void SPluginCategories::Construct( const FArguments& Args, const TSharedRef< SPluginsEditor > Owner )
+void SPluginCategories::Construct( const FArguments& Args, const TSharedRef< SPluginBrowser > Owner )
 {
 	OwnerWeak = Owner;
 
 	RootPluginCategories.Reset();
 
 	// Add the root level categories
-	TSharedRef<FPluginCategoryTreeItem> InstalledCategory = MakeShareable(new FPluginCategoryTreeItem(NULL, TEXT("Installed"), LOCTEXT("InstalledCategoryName", "Installed")));
+	TSharedRef<FPluginCategory> InstalledCategory = MakeShareable(new FPluginCategory(NULL, TEXT("Installed"), LOCTEXT("InstalledCategoryName", "Installed")));
 	RootPluginCategories.Add( InstalledCategory );
-	TSharedRef<FPluginCategoryTreeItem> ProjectCategory = MakeShareable(new FPluginCategoryTreeItem(NULL, TEXT("Project"), LOCTEXT("ProjectCategoryName", "Project")));
+	TSharedRef<FPluginCategory> ProjectCategory = MakeShareable(new FPluginCategory(NULL, TEXT("Project"), LOCTEXT("ProjectCategoryName", "Project")));
 	RootPluginCategories.Add( ProjectCategory );
 
 	PluginCategoryTreeView =
-		SNew( SPluginCategoryTreeView )
+		SNew( STreeView<TSharedPtr<FPluginCategory>> )
 
 		// For now we only support selecting a single folder in the tree
 		.SelectionMode( ESelectionMode::Single )
@@ -56,7 +55,7 @@ SPluginCategories::~SPluginCategories()
 
 
 /** @return Gets the owner of this list */
-SPluginsEditor& SPluginCategories::GetOwner()
+SPluginBrowser& SPluginCategories::GetOwner()
 {
 	return *OwnerWeak.Pin();
 }
@@ -66,7 +65,7 @@ void SPluginCategories::RebuildAndFilterCategoryTree()
 {
 	// Get a plugin from the currently selected category, so we can track it if it's removed
 	IPlugin* TrackPlugin = nullptr;
-	for(FPluginCategoryTreeItemPtr SelectedItem: PluginCategoryTreeView->GetSelectedItems())
+	for(TSharedPtr<FPluginCategory> SelectedItem: PluginCategoryTreeView->GetSelectedItems())
 	{
 		if(SelectedItem->Plugins.Num() > 0)
 		{
@@ -76,20 +75,20 @@ void SPluginCategories::RebuildAndFilterCategoryTree()
 	}
 
 	// Clear the list of plugins in each current category
-	for(FPluginCategoryTreeItemPtr RootCategory: RootPluginCategories)
+	for(TSharedPtr<FPluginCategory> RootCategory: RootPluginCategories)
 	{
-		for(FPluginCategoryTreeItemPtr Category: RootCategory->AccessSubCategories())
+		for(TSharedPtr<FPluginCategory> Category: RootCategory->SubCategories)
 		{
 			Category->Plugins.Reset();
 		}
 	}
 
 	// Add all the known plugins into categories
-	FPluginCategoryTreeItemPtr SelectCategory;
+	TSharedPtr<FPluginCategory> SelectCategory;
 	for(IPlugin* Plugin: IPluginManager::Get().GetDiscoveredPlugins())
 	{
 		// Figure out which base category this plugin belongs in
-		FPluginCategoryTreeItemPtr RootCategory;
+		TSharedPtr<FPluginCategory> RootCategory;
 		if( Plugin->GetLoadedFrom() == EPluginLoadedFrom::Engine )
 		{
 			RootCategory = RootPluginCategories[0];
@@ -107,10 +106,10 @@ void SPluginCategories::RebuildAndFilterCategoryTree()
 		}
 
 		// Locate this category at the level we're at in the hierarchy
-		FPluginCategoryTreeItemPtr FoundCategory = NULL;
-		for(FPluginCategoryTreeItemPtr TestCategory: RootCategory->SubCategories)
+		TSharedPtr<FPluginCategory> FoundCategory = NULL;
+		for(TSharedPtr<FPluginCategory> TestCategory: RootCategory->SubCategories)
 		{
-			if(TestCategory->GetCategoryName() == CategoryName)
+			if(TestCategory->Name == CategoryName)
 			{
 				FoundCategory = TestCategory;
 				break;
@@ -120,12 +119,12 @@ void SPluginCategories::RebuildAndFilterCategoryTree()
 		if( !FoundCategory.IsValid() )
 		{
 			//@todo Allow for properly localized category names [3/7/2014 justin.sargent]
-			FoundCategory = MakeShareable(new FPluginCategoryTreeItem(RootCategory, CategoryName, FText::FromString(CategoryName)));
+			FoundCategory = MakeShareable(new FPluginCategory(RootCategory, CategoryName, FText::FromString(CategoryName)));
 			RootCategory->SubCategories.Add( FoundCategory );
 		}
 			
 		// Associate the plugin with the category
-		FoundCategory->AddPlugin(Plugin);
+		FoundCategory->Plugins.Add(Plugin);
 
 		// Update the selection if this is the plugin we're tracking
 		if(TrackPlugin == Plugin)
@@ -135,7 +134,7 @@ void SPluginCategories::RebuildAndFilterCategoryTree()
 	}
 
 	// Remove any empty categories, keeping track of which items are still selected
-	for(FPluginCategoryTreeItemPtr RootCategory: RootPluginCategories)
+	for(TSharedPtr<FPluginCategory> RootCategory: RootPluginCategories)
 	{
 		for(int32 Idx = 0; Idx < RootCategory->SubCategories.Num(); Idx++)
 		{
@@ -147,9 +146,9 @@ void SPluginCategories::RebuildAndFilterCategoryTree()
 	}
 
 	// Sort every single category alphabetically
-	for(FPluginCategoryTreeItemPtr RootCategory: RootPluginCategories)
+	for(TSharedPtr<FPluginCategory> RootCategory: RootPluginCategories)
 	{
-		RootCategory->SubCategories.Sort([](const FPluginCategoryTreeItemPtr& A, const FPluginCategoryTreeItemPtr& B) -> bool { return A->CategoryName < B->CategoryName; });
+		RootCategory->SubCategories.Sort([](const TSharedPtr<FPluginCategory>& A, const TSharedPtr<FPluginCategory>& B) -> bool { return A->DisplayName.CompareTo(B->DisplayName) < 0; });
 	}
 
 	// Refresh the view
@@ -170,31 +169,30 @@ void SPluginCategories::RebuildAndFilterCategoryTree()
 	}
 }
 
-TSharedRef<ITableRow> SPluginCategories::PluginCategoryTreeView_OnGenerateRow( FPluginCategoryTreeItemPtr Item, const TSharedRef<STableViewBase>& OwnerTable )
+TSharedRef<ITableRow> SPluginCategories::PluginCategoryTreeView_OnGenerateRow( TSharedPtr<FPluginCategory> Item, const TSharedRef<STableViewBase>& OwnerTable )
 {
 	return
-		SNew( STableRow< FPluginCategoryTreeItemPtr >, OwnerTable )
+		SNew( STableRow<TSharedPtr<FPluginCategory>>, OwnerTable )
 		[
-			SNew( SPluginCategoryTreeItem, SharedThis( this ), Item.ToSharedRef() )
+			SNew( SPluginCategoryTreeItem, Item.ToSharedRef() )
 		];
 }
 
 
-void SPluginCategories::PluginCategoryTreeView_OnGetChildren( FPluginCategoryTreeItemPtr Item, TArray< FPluginCategoryTreeItemPtr >& OutChildren )
+void SPluginCategories::PluginCategoryTreeView_OnGetChildren(TSharedPtr<FPluginCategory> Item, TArray<TSharedPtr<FPluginCategory>>& OutChildren )
 {
-	const auto& SubCategories = Item->GetSubCategories();
-	OutChildren.Append( SubCategories );
+	OutChildren.Append(Item->SubCategories);
 }
 
 
-void SPluginCategories::PluginCategoryTreeView_OnSelectionChanged( FPluginCategoryTreeItemPtr Item, ESelectInfo::Type SelectInfo )
+void SPluginCategories::PluginCategoryTreeView_OnSelectionChanged(TSharedPtr<FPluginCategory> Item, ESelectInfo::Type SelectInfo )
 {
 	// Selection changed, which may affect which plugins are displayed in the list.  We need to invalidate the list.
 	OwnerWeak.Pin()->OnCategorySelectionChanged();
 }
 
 
-FPluginCategoryTreeItemPtr SPluginCategories::GetSelectedCategory() const
+TSharedPtr<FPluginCategory> SPluginCategories::GetSelectedCategory() const
 {
 	if( PluginCategoryTreeView.IsValid() )
 	{
@@ -209,7 +207,7 @@ FPluginCategoryTreeItemPtr SPluginCategories::GetSelectedCategory() const
 	return NULL;
 }
 
-void SPluginCategories::SelectCategory( const FPluginCategoryTreeItemPtr& CategoryToSelect )
+void SPluginCategories::SelectCategory( const TSharedPtr<FPluginCategory>& CategoryToSelect )
 {
 	if( ensure( PluginCategoryTreeView.IsValid() ) )
 	{
@@ -217,7 +215,7 @@ void SPluginCategories::SelectCategory( const FPluginCategoryTreeItemPtr& Catego
 	}
 }
 
-bool SPluginCategories::IsItemExpanded( const FPluginCategoryTreeItemPtr Item ) const
+bool SPluginCategories::IsItemExpanded( const TSharedPtr<FPluginCategory> Item ) const
 {
 	return PluginCategoryTreeView->IsItemExpanded( Item );
 }
diff --git a/Engine/Plugins/Editor/PluginBrowser/Source/PluginBrowser/Private/SPluginCategories.h b/Engine/Plugins/Editor/PluginBrowser/Source/PluginBrowser/Private/SPluginCategories.h
--- a/Engine/Plugins/Editor/PluginBrowser/Source/PluginBrowser/Private/SPluginCategories.h
+++ b/Engine/Plugins/Editor/PluginBrowser/Source/PluginBrowser/Private/SPluginCategories.h
@@ -2,10 +2,7 @@
 
 #pragma once
 
-#include "PluginCategoryTreeItem.h"
-
-typedef STreeView< FPluginCategoryTreeItemPtr > SPluginCategoryTreeView;
-
+#include "SPluginCategoryTreeItem.h"
 
 /**
  * Tree view that displays all of the plugin categories and allows the user to switch views
@@ -23,22 +20,22 @@ public:
 
 
 	/** Widget constructor */
-	void Construct( const FArguments& Args, const TSharedRef< class SPluginsEditor > Owner );
+	void Construct( const FArguments& Args, const TSharedRef< class SPluginBrowser > Owner );
 
 	/** Destructor */
 	~SPluginCategories();
 
 	/** @return Gets the owner of this categories tree */
-	SPluginsEditor& GetOwner();
+	SPluginBrowser& GetOwner();
 
 	/** @return Returns the currently selected category item */
-	FPluginCategoryTreeItemPtr GetSelectedCategory() const;
+	TSharedPtr<FPluginCategory> GetSelectedCategory() const;
 
 	/** Selects the specified category */
-	void SelectCategory( const FPluginCategoryTreeItemPtr& CategoryToSelect );
+	void SelectCategory(const TSharedPtr<FPluginCategory>& CategoryToSelect);
 
 	/** @return Returns true if the specified item is currently expanded in the tree */
-	bool IsItemExpanded( const FPluginCategoryTreeItemPtr Item ) const;
+	bool IsItemExpanded(const TSharedPtr<FPluginCategory> Item) const;
 
 	/** Signal that the categories list needs to be refreshed */
 	void SetNeedsRefresh();
@@ -46,13 +43,13 @@ public:
 private:
 	
 	/** Called to generate a widget for the specified tree item */
-	TSharedRef<ITableRow> PluginCategoryTreeView_OnGenerateRow( FPluginCategoryTreeItemPtr Item, const TSharedRef<STableViewBase>& OwnerTable );
+	TSharedRef<ITableRow> PluginCategoryTreeView_OnGenerateRow(TSharedPtr<FPluginCategory> Item, const TSharedRef<STableViewBase>& OwnerTable);
 
 	/** Given a tree item, fills an array of child items */
-	void PluginCategoryTreeView_OnGetChildren( FPluginCategoryTreeItemPtr Item, TArray< FPluginCategoryTreeItemPtr >& OutChildren );
+	void PluginCategoryTreeView_OnGetChildren(TSharedPtr<FPluginCategory> Item, TArray<TSharedPtr<FPluginCategory>>& OutChildren);
 
 	/** Called when the user clicks on a plugin item, or when selection changes by some other means */
-	void PluginCategoryTreeView_OnSelectionChanged( FPluginCategoryTreeItemPtr Item, ESelectInfo::Type SelectInfo );
+	void PluginCategoryTreeView_OnSelectionChanged(TSharedPtr<FPluginCategory> Item, ESelectInfo::Type SelectInfo);
 
 	/** Rebuilds the category tree from scratch */
 	void RebuildAndFilterCategoryTree();
@@ -63,12 +60,12 @@ private:
 private:
 
 	/** Weak pointer back to its owner */
-	TWeakPtr< class SPluginsEditor > OwnerWeak;
+	TWeakPtr< class SPluginBrowser > OwnerWeak;
 
 	/** The tree view widget for our plugin categories tree */
-	TSharedPtr< SPluginCategoryTreeView > PluginCategoryTreeView;
+	TSharedPtr<STreeView<TSharedPtr<FPluginCategory>>> PluginCategoryTreeView;
 
 	/** Root list of categories */
-	TArray< FPluginCategoryTreeItemPtr > RootPluginCategories;
+	TArray<TSharedPtr<FPluginCategory>> RootPluginCategories;
 };
 
diff --git a/Engine/Plugins/Editor/PluginBrowser/Source/PluginBrowser/Private/SPluginCategoryTreeItem.cpp b/Engine/Plugins/Editor/PluginBrowser/Source/PluginBrowser/Private/SPluginCategoryTreeItem.cpp
--- a/Engine/Plugins/Editor/PluginBrowser/Source/PluginBrowser/Private/SPluginCategoryTreeItem.cpp
+++ b/Engine/Plugins/Editor/PluginBrowser/Source/PluginBrowser/Private/SPluginCategoryTreeItem.cpp
@@ -1,8 +1,7 @@
 // Copyright 1998-2015 Epic Games, Inc. All Rights Reserved.
 
-#include "PluginsEditorPrivatePCH.h"
+#include "PluginBrowserPrivatePCH.h"
 #include "SPluginCategoryTreeItem.h"
-#include "PluginCategoryTreeItem.h"
 #include "SPluginCategories.h"
 #include "PluginStyle.h"
 
@@ -10,18 +9,15 @@
 #define LOCTEXT_NAMESPACE "PluginCategoryTreeItem"
 
 
-void SPluginCategoryTreeItem::Construct( const FArguments& Args, const TSharedRef< SPluginCategories > Owner, const TSharedRef< class FPluginCategoryTreeItem >& Item )
+void SPluginCategoryTreeItem::Construct(const FArguments& Args, const TSharedRef<FPluginCategory>& InCategory)
 {
-	OwnerWeak = Owner;
-	ItemData = Item;
+	Category = InCategory;
 
 	const float CategoryIconSize = FPluginStyle::Get()->GetFloat("CategoryTreeItem.IconSize");
 	const float PaddingAmount = FPluginStyle::Get()->GetFloat("CategoryTreeItem.PaddingAmount");;
 
-	// @todo plugedit: Show a badge when at least one plugin in this category is not at the latest version
-
 	// Figure out which font size to use
-	const auto bIsRootItem = !Item->GetParentCategory().IsValid();
+	const auto bIsRootItem = !Category->ParentCategory.IsValid();
 
 	ChildSlot
 	[ 
@@ -52,7 +48,7 @@ void SPluginCategoryTreeItem::Construct( const FArguments& Args, const TSharedRe
 			.VAlign(VAlign_Center)
 			[
 				SNew( STextBlock )
-				.Text( Item->GetCategoryDisplayName() )
+				.Text( Category->DisplayName )
 				.TextStyle( FPluginStyle::Get(), bIsRootItem ? "CategoryTreeItem.Root.Text" : "CategoryTreeItem.Text" )
 			]
 			
@@ -65,9 +61,9 @@ void SPluginCategoryTreeItem::Construct( const FArguments& Args, const TSharedRe
 				SNew( STextBlock )
 
 				// Only display if at there is least one plugin is in this category
-				.Visibility( Item->GetPlugins().Num() > 0 ? EVisibility::Visible : EVisibility::Collapsed )
+				.Visibility( Category->Plugins.Num() > 0 ? EVisibility::Visible : EVisibility::Collapsed )
 
-				.Text( FText::Format( LOCTEXT( "NumberOfPluginsWrapper", "({0})" ), FText::AsNumber( Item->GetPlugins().Num() ) ) )
+				.Text( FText::Format( LOCTEXT( "NumberOfPluginsWrapper", "({0})" ), FText::AsNumber( Category->Plugins.Num() ) ) )
 				.TextStyle( FPluginStyle::Get(), bIsRootItem ? "CategoryTreeItem.Root.PluginCountText" : "CategoryTreeItem.PluginCountText" )
 			]
 		]
@@ -77,23 +73,18 @@ void SPluginCategoryTreeItem::Construct( const FArguments& Args, const TSharedRe
 
 const FSlateBrush* SPluginCategoryTreeItem::GetIconBrush() const
 {
-	// @todo plugedit: Need support for styles embedded into plugins!
-	const FSlateBrush* IconBrush;
-
-	if(ItemData->GetParentCategory().IsValid())
+	if(Category->ParentCategory.IsValid())
 	{
-		IconBrush = FPluginStyle::Get()->GetBrush( "CategoryTreeItem.LeafItemWithPlugin" );
+		return FPluginStyle::Get()->GetBrush( "CategoryTreeItem.LeafItemWithPlugin" );
 	}
-	else if ( ItemData->GetCategoryName() == TEXT("Installed"))
+	else if (Category->Name == TEXT("Installed"))
 	{
-		IconBrush = FPluginStyle::Get()->GetBrush( "CategoryTreeItem.BuiltIn" );
+		return FPluginStyle::Get()->GetBrush( "CategoryTreeItem.BuiltIn" );
 	}
 	else
 	{
-		IconBrush = FPluginStyle::Get()->GetBrush( "CategoryTreeItem.Installed" );
+		return FPluginStyle::Get()->GetBrush( "CategoryTreeItem.Installed" );
 	}
-
-	return IconBrush;
 }
 
 
diff --git a/Engine/Plugins/Editor/PluginBrowser/Source/PluginBrowser/Private/SPluginCategoryTreeItem.h b/Engine/Plugins/Editor/PluginBrowser/Source/PluginBrowser/Private/SPluginCategoryTreeItem.h
--- a/Engine/Plugins/Editor/PluginBrowser/Source/PluginBrowser/Private/SPluginCategoryTreeItem.h
+++ b/Engine/Plugins/Editor/PluginBrowser/Source/PluginBrowser/Private/SPluginCategoryTreeItem.h
@@ -2,6 +2,36 @@
 
 #pragma once
 
+/**
+ * Represents a category in the plugin category tree
+ */
+class FPluginCategory
+{
+public:
+	/** Parent category item or NULL if this is a root category */
+	TWeakPtr<FPluginCategory> ParentCategory;
+
+	/** Name of the category */
+	FString Name;
+
+	/** Display name of the category */
+	FText DisplayName;
+
+	/** Child categories */
+	TArray<TSharedPtr<FPluginCategory>> SubCategories;
+
+	/** Plugins in this category */
+	TArray<IPlugin*> Plugins;
+
+	/** Constructor for FPluginCategory */
+	FPluginCategory(TSharedPtr<FPluginCategory> InParentCategory, const FString& InName, const FText& InDisplayName)
+		: ParentCategory(InParentCategory),
+		  Name(InName),
+		  DisplayName(InDisplayName)
+	{
+	}
+};
+
 /**
  * Widget that represents a single category in the category tree view
  */
@@ -13,23 +43,17 @@ class SPluginCategoryTreeItem : public SCompoundWidget
 
 	SLATE_END_ARGS()
 
-
 	/** Widget constructor */
-	void Construct( const FArguments& Args, const TSharedRef< class SPluginCategories > Owner, const TSharedRef< class FPluginCategoryTreeItem >& Item );
-
+	void Construct(const FArguments& Args, const TSharedRef<FPluginCategory>& Item);
 
 private:
 
 	/** @return Gets the icon brush to use for this item's current state */
 	const FSlateBrush* GetIconBrush() const;
 
-
 private:
 
 	/** The item we're representing the in tree */
-	TSharedPtr< class FPluginCategoryTreeItem > ItemData;
-
-	/** Weak pointer back to its owner */
-	TWeakPtr< class SPluginCategories > OwnerWeak;
+	TSharedPtr<FPluginCategory> Category;
 };
 
diff --git a/Engine/Plugins/Editor/PluginBrowser/Source/PluginBrowser/Private/SPluginList.cpp b/Engine/Plugins/Editor/PluginBrowser/Source/PluginBrowser/Private/SPluginList.cpp
--- a/Engine/Plugins/Editor/PluginBrowser/Source/PluginBrowser/Private/SPluginList.cpp
+++ b/Engine/Plugins/Editor/PluginBrowser/Source/PluginBrowser/Private/SPluginList.cpp
@@ -1,18 +1,16 @@
 // Copyright 1998-2015 Epic Games, Inc. All Rights Reserved.
 
-#include "PluginsEditorPrivatePCH.h"
+#include "PluginBrowserPrivatePCH.h"
 #include "SPluginList.h"
-#include "SPluginsEditor.h"
-#include "PluginListItem.h"
+#include "SPluginBrowser.h"
 #include "SPluginListTile.h"
 #include "IPluginManager.h"
-#include "PluginCategoryTreeItem.h"
 
 
 #define LOCTEXT_NAMESPACE "PluginList"
 
 
-void SPluginList::Construct( const FArguments& Args, const TSharedRef< SPluginsEditor > Owner )
+void SPluginList::Construct( const FArguments& Args, const TSharedRef< SPluginBrowser > Owner )
 {
 	OwnerWeak = Owner;
 
@@ -25,7 +23,7 @@ void SPluginList::Construct( const FArguments& Args, const TSharedRef< SPluginsE
 	// @todo plugedit: Have optional compact version with only plugin icon + name + version?  Only expand selected?
 
 	PluginListView = 
-		SNew( SPluginListView )
+		SNew( SListView<TSharedRef<IPlugin*>> )
 
 		.SelectionMode( ESelectionMode::None )		// No need to select plugins!
 
@@ -39,7 +37,7 @@ void SPluginList::Construct( const FArguments& Args, const TSharedRef< SPluginsE
 
 SPluginList::~SPluginList()
 {
-	const TSharedPtr< SPluginsEditor > Owner( OwnerWeak.Pin() );
+	const TSharedPtr< SPluginBrowser > Owner( OwnerWeak.Pin() );
 	if( Owner.IsValid() )
 	{
 		Owner->GetPluginTextFilter().OnChanged().RemoveAll( this );
@@ -48,38 +46,23 @@ SPluginList::~SPluginList()
 
 
 /** @return Gets the owner of this list */
-SPluginsEditor& SPluginList::GetOwner()
+SPluginBrowser& SPluginList::GetOwner()
 {
 	return *OwnerWeak.Pin();
 }
 
 
-TSharedRef<ITableRow> SPluginList::PluginListView_OnGenerateRow( FPluginListItemPtr Item, const TSharedRef<STableViewBase>& OwnerTable )
+TSharedRef<ITableRow> SPluginList::PluginListView_OnGenerateRow( TSharedRef<IPlugin*> Item, const TSharedRef<STableViewBase>& OwnerTable )
 {
 	return
-		SNew( STableRow< FPluginListItemPtr >, OwnerTable )
+		SNew( STableRow< TSharedRef<IPlugin*> >, OwnerTable )
 		[
-			SNew( SPluginListTile, SharedThis( this ), Item.ToSharedRef() )
+			SNew( SPluginListTile, SharedThis( this ), Item.Get() )
 		];
 }
 
 
 
-void SPluginList::GetPlugins( const TSharedPtr< FPluginCategoryTreeItem >& Category )
-{
-	// Add plugins from this category
-	for(IPlugin* Plugin: Category->GetPlugins())
-	{
-		if(OwnerWeak.Pin()->GetPluginTextFilter().PassesFilter(Plugin))
-		{
-			const TSharedRef< FPluginListItem > NewItem( new FPluginListItem() );
-			NewItem->Plugin = Plugin;
-			PluginListItems.Add( NewItem );
-		}
-	}
-}
-
-
 void SPluginList::RebuildAndFilterPluginList()
 {
 	// Build up the initial list of plugins
@@ -87,19 +70,25 @@ void SPluginList::RebuildAndFilterPluginList()
 		PluginListItems.Reset();
 
 		// Get the currently selected category
-		const auto& SelectedCategory = OwnerWeak.Pin()->GetSelectedCategory();
+		const TSharedPtr<FPluginCategory>& SelectedCategory = OwnerWeak.Pin()->GetSelectedCategory();
 		if( SelectedCategory.IsValid() )
 		{
-			GetPlugins( SelectedCategory );
+			for(IPlugin* Plugin: SelectedCategory->Plugins)
+			{
+				if(OwnerWeak.Pin()->GetPluginTextFilter().PassesFilter(Plugin))
+				{
+					PluginListItems.Add(TSharedRef<IPlugin*>(new IPlugin*(Plugin)));
+				}
+			}
 		}
 
 		// Sort the plugins alphabetically
 		{
 			struct FPluginListItemSorter
 			{
-				bool operator()( const FPluginListItemPtr& A, const FPluginListItemPtr& B ) const
+				bool operator()( const TSharedRef<IPlugin*>& A, const TSharedRef<IPlugin*>& B ) const
 				{
-					return A->Plugin->GetDescriptor().FriendlyName < B->Plugin->GetDescriptor().FriendlyName;
+					return A.Get()->GetDescriptor().FriendlyName < B.Get()->GetDescriptor().FriendlyName;
 				}
 			};
 			PluginListItems.Sort( FPluginListItemSorter() );
diff --git a/Engine/Plugins/Editor/PluginBrowser/Source/PluginBrowser/Private/SPluginList.h b/Engine/Plugins/Editor/PluginBrowser/Source/PluginBrowser/Private/SPluginList.h
--- a/Engine/Plugins/Editor/PluginBrowser/Source/PluginBrowser/Private/SPluginList.h
+++ b/Engine/Plugins/Editor/PluginBrowser/Source/PluginBrowser/Private/SPluginList.h
@@ -2,11 +2,6 @@
 
 #pragma once
 
-#include "PluginListItem.h"
-
-typedef SListView< FPluginListItemPtr > SPluginListView;
-
-
 /**
  * A filtered list of plugins, driven by a category selection
  */
@@ -19,13 +14,13 @@ class SPluginList : public SCompoundWidget
 	SLATE_END_ARGS()
 
 	/** Widget constructor */
-	void Construct( const FArguments& Args, const TSharedRef< class SPluginsEditor > Owner );
+	void Construct( const FArguments& Args, const TSharedRef< class SPluginBrowser > Owner );
 
 	/** Destructor */
 	~SPluginList();
 
 	/** @return Gets the owner of this list */
-	SPluginsEditor& GetOwner();
+	SPluginBrowser& GetOwner();
 
 	/** Called to invalidate the list */
 	void SetNeedsRefresh();
@@ -36,11 +31,8 @@ private:
 	/** Called when the plugin text filter has changed what its filtering */
 	void OnPluginTextFilterChanged();
 
-	/** Recursively gathers plugins from a category and all sub-categories.  Updates the PluginListItems structure. */
-	void GetPlugins( const TSharedPtr< class FPluginCategoryTreeItem >& Category );
-
 	/** Called to generate a widget for the specified list item */
-	TSharedRef<ITableRow> PluginListView_OnGenerateRow( FPluginListItemPtr Item, const TSharedRef<STableViewBase>& OwnerTable );
+	TSharedRef<ITableRow> PluginListView_OnGenerateRow(TSharedRef<IPlugin*> Item, const TSharedRef<STableViewBase>& OwnerTable );
 
 	/** Rebuilds the list of plugins from scratch and applies filtering. */
 	void RebuildAndFilterPluginList();
@@ -51,13 +43,13 @@ private:
 private:
 
 	/** Weak pointer back to its owner */
-	TWeakPtr< class SPluginsEditor > OwnerWeak;
+	TWeakPtr< class SPluginBrowser > OwnerWeak;
 
 	/** The list view widget for our plugins list */
-	TSharedPtr< SPluginListView > PluginListView;
+	TSharedPtr<SListView<TSharedRef<IPlugin*>>> PluginListView;
 
 	/** List of everything that we want to display in the plugin list */
-	TArray< FPluginListItemPtr > PluginListItems;
+	TArray<TSharedRef<IPlugin*>> PluginListItems;
 
 	/** Whether the active timer to refresh the list is registered */
 	bool bIsActiveTimerRegistered;
diff --git a/Engine/Plugins/Editor/PluginBrowser/Source/PluginBrowser/Private/SPluginListTile.cpp b/Engine/Plugins/Editor/PluginBrowser/Source/PluginBrowser/Private/SPluginListTile.cpp
--- a/Engine/Plugins/Editor/PluginBrowser/Source/PluginBrowser/Private/SPluginListTile.cpp
+++ b/Engine/Plugins/Editor/PluginBrowser/Source/PluginBrowser/Private/SPluginListTile.cpp
@@ -1,9 +1,8 @@
 // Copyright 1998-2015 Epic Games, Inc. All Rights Reserved.
 
-#include "PluginsEditorPrivatePCH.h"
+#include "PluginBrowserPrivatePCH.h"
 #include "SPluginListTile.h"
-#include "PluginListItem.h"
-#include "SPluginsEditor.h"
+#include "SPluginBrowser.h"
 #include "SPluginList.h"
 #include "PluginStyle.h"
 #include "GameProjectGenerationModule.h"
@@ -17,10 +16,10 @@
 #define LOCTEXT_NAMESPACE "PluginListTile"
 
 
-void SPluginListTile::Construct( const FArguments& Args, const TSharedRef<SPluginList> Owner, const TSharedRef<class FPluginListItem>& Item )
+void SPluginListTile::Construct( const FArguments& Args, const TSharedRef<SPluginList> Owner, IPlugin* InPlugin )
 {
 	OwnerWeak = Owner;
-	ItemData = Item;
+	Plugin = InPlugin;
 
 	const float PaddingAmount = FPluginStyle::Get()->GetFloat( "PluginTile.Padding" );
 	const float ThumbnailImageSize = FPluginStyle::Get()->GetFloat( "PluginTile.ThumbnailImageSize" );
@@ -31,7 +30,6 @@ void SPluginListTile::Construct( const FArguments& Args, const TSharedRef<SPlugi
 
 	// @todo plugedit: Maybe we should do the FileExists check ONCE at plugin load time and not at query time
 
-	IPlugin* Plugin = ItemData->Plugin;
 	const FPluginDescriptor& PluginDescriptor = Plugin->GetDescriptor();
 
 	// Plugin thumbnail image
@@ -312,13 +310,13 @@ void SPluginListTile::Construct( const FArguments& Args, const TSharedRef<SPlugi
 ECheckBoxState SPluginListTile::IsPluginEnabled() const
 {
 	FPluginBrowserModule& PluginBrowserModule = FPluginBrowserModule::Get();
-	if(PluginBrowserModule.PendingEnablePlugins.Contains(ItemData->Plugin))
+	if(PluginBrowserModule.PendingEnablePlugins.Contains(Plugin))
 	{
-		return PluginBrowserModule.PendingEnablePlugins[ItemData->Plugin]? ECheckBoxState::Checked : ECheckBoxState::Unchecked;;
+		return PluginBrowserModule.PendingEnablePlugins[Plugin]? ECheckBoxState::Checked : ECheckBoxState::Unchecked;;
 	}
 	else
 	{
-		return ItemData->Plugin->IsEnabled()? ECheckBoxState::Checked : ECheckBoxState::Unchecked;
+		return Plugin->IsEnabled()? ECheckBoxState::Checked : ECheckBoxState::Unchecked;
 	}
 }
 
@@ -326,7 +324,7 @@ void SPluginListTile::OnEnablePluginCheckboxChanged(ECheckBoxState NewCheckedSta
 {
 	const bool bNewEnabledState = (NewCheckedState == ECheckBoxState::Checked);
 
-	const FPluginDescriptor& PluginDescriptor = ItemData->Plugin->GetDescriptor();
+	const FPluginDescriptor& PluginDescriptor = Plugin->GetDescriptor();
 	if (bNewEnabledState && PluginDescriptor.bIsBetaVersion)
 	{
 		FText WarningMessage = FText::Format(
@@ -343,20 +341,20 @@ void SPluginListTile::OnEnablePluginCheckboxChanged(ECheckBoxState NewCheckedSta
 	FGameProjectGenerationModule::Get().TryMakeProjectFileWriteable(FPaths::GetProjectFilePath());
 	FText FailMessage;
 
-	if (!IProjectManager::Get().SetPluginEnabled(ItemData->Plugin->GetName(), bNewEnabledState, FailMessage))
+	if (!IProjectManager::Get().SetPluginEnabled(Plugin->GetName(), bNewEnabledState, FailMessage))
 	{
 		FMessageDialog::Open(EAppMsgType::Ok, FailMessage);
 	}
 	else
 	{
 		FPluginBrowserModule& PluginBrowserModule = FPluginBrowserModule::Get();
-		if(ItemData->Plugin->IsEnabled() == bNewEnabledState)
+		if(Plugin->IsEnabled() == bNewEnabledState)
 		{
-			PluginBrowserModule.PendingEnablePlugins.Remove(ItemData->Plugin);
+			PluginBrowserModule.PendingEnablePlugins.Remove(Plugin);
 		}
 		else
 		{
-			PluginBrowserModule.PendingEnablePlugins.FindOrAdd(ItemData->Plugin) = bNewEnabledState;
+			PluginBrowserModule.PendingEnablePlugins.FindOrAdd(Plugin) = bNewEnabledState;
 		}
 	}
 
@@ -365,15 +363,15 @@ void SPluginListTile::OnEnablePluginCheckboxChanged(ECheckBoxState NewCheckedSta
 
 EVisibility SPluginListTile::GetAuthoringButtonsVisibility() const
 {
-	return (FApp::IsEngineInstalled() && ItemData->Plugin->GetLoadedFrom() == EPluginLoadedFrom::Engine)? EVisibility::Hidden : EVisibility::Visible;
+	return (FApp::IsEngineInstalled() && Plugin->GetLoadedFrom() == EPluginLoadedFrom::Engine)? EVisibility::Hidden : EVisibility::Visible;
 }
 
 void SPluginListTile::OnEditPlugin()
 {
 	// Construct the plugin metadata object using the descriptor for this plugin
 	UPluginMetadataObject* MetadataObject = NewObject<UPluginMetadataObject>();
-	MetadataObject->TargetIconPath = ItemData->Plugin->GetBaseDir() / TEXT("Resources/Icon128.png");
-	MetadataObject->PopulateFromDescriptor(ItemData->Plugin->GetDescriptor());
+	MetadataObject->TargetIconPath = Plugin->GetBaseDir() / TEXT("Resources/Icon128.png");
+	MetadataObject->PopulateFromDescriptor(Plugin->GetDescriptor());
 	MetadataObject->AddToRoot();
 
 	// Create a property view
@@ -401,7 +399,7 @@ void SPluginListTile::OnEditPlugin()
 				[
 					SNew(STextBlock)
 					.Font(FPluginStyle::Get()->GetFontStyle(TEXT("PluginMetadataNameFont")))
-					.Text(FText::FromString(ItemData->Plugin->GetName()))
+					.Text(FText::FromString(Plugin->GetName()))
 				]
 
 				+ SVerticalBox::Slot()
@@ -428,7 +426,6 @@ void SPluginListTile::OnEditPlugin()
 
 FReply SPluginListTile::OnEditPluginFinished(UPluginMetadataObject* MetadataObject)
 {
-	IPlugin* Plugin = ItemData->Plugin;
 	FPluginDescriptor OldDescriptor = Plugin->GetDescriptor();
 
 	// Update the descriptor with the new metadata
@@ -460,7 +457,7 @@ FReply SPluginListTile::OnEditPluginFinished(UPluginMetadataObject* MetadataObje
 
 		// Write to the file and update the in-memory metadata
 		FText FailReason;
-		if(!ItemData->Plugin->UpdateDescriptor(NewDescriptor, FailReason))
+		if(!Plugin->UpdateDescriptor(NewDescriptor, FailReason))
 		{
 			FMessageDialog::Open(EAppMsgType::Ok, FailReason);
 		}
diff --git a/Engine/Plugins/Editor/PluginBrowser/Source/PluginBrowser/Private/SPluginListTile.h b/Engine/Plugins/Editor/PluginBrowser/Source/PluginBrowser/Private/SPluginListTile.h
--- a/Engine/Plugins/Editor/PluginBrowser/Source/PluginBrowser/Private/SPluginListTile.h
+++ b/Engine/Plugins/Editor/PluginBrowser/Source/PluginBrowser/Private/SPluginListTile.h
@@ -15,7 +15,7 @@ class SPluginListTile : public SCompoundWidget
 
 
 	/** Widget constructor */
-	void Construct( const FArguments& Args, const TSharedRef< class SPluginList > Owner, const TSharedRef< class FPluginListItem >& Item );
+	void Construct( const FArguments& Args, const TSharedRef< class SPluginList > Owner, IPlugin* Plugin );
 
 
 private:
@@ -38,7 +38,7 @@ private:
 private:
 
 	/** The item we're representing the in tree */
-	TSharedPtr< class FPluginListItem > ItemData;
+	IPlugin* Plugin;
 
 	/** Weak pointer back to its owner */
 	TWeakPtr< class SPluginList > OwnerWeak;
diff --git a/Engine/Plugins/Editor/PluginBrowser/Source/PluginBrowser/Private/SPluginsEditor.cpp b/Engine/Plugins/Editor/PluginBrowser/Source/PluginBrowser/Private/SPluginsEditor.cpp
deleted file mode 100644
diff --git a/Engine/Plugins/Editor/PluginBrowser/Source/PluginBrowser/Private/SPluginsEditor.h b/Engine/Plugins/Editor/PluginBrowser/Source/PluginBrowser/Private/SPluginsEditor.h
deleted file mode 100644
diff --git a/Engine/Plugins/Editor/PluginBrowser/Source/PluginBrowser/Public/IPluginBrowser.h b/Engine/Plugins/Editor/PluginBrowser/Source/PluginBrowser/Public/IPluginBrowser.h
new file mode 100644
--- /dev/null
+++ b/Engine/Plugins/Editor/PluginBrowser/Source/PluginBrowser/Public/IPluginBrowser.h
@@ -0,0 +1,37 @@
+// Copyright 1998-2015 Epic Games, Inc. All Rights Reserved.
+
+#pragma once
+
+#include "ModuleManager.h"
+#include "Editor/UnrealEd/Public/Features/IPluginsEditorFeature.h"
+
+/**
+ * The public interface to this module
+ */
+class IPluginBrowser : public IModuleInterface, public IPluginsEditorFeature
+{
+
+public:
+
+	/**
+	 * Singleton-like access to this module's interface.  This is just for convenience!
+	 * Beware of calling this during the shutdown phase, though.  Your module might have been unloaded already.
+	 *
+	 * @return Returns singleton instance, loading the module on demand if needed
+	 */
+	static inline IPluginBrowser& Get()
+	{
+		return FModuleManager::LoadModuleChecked< IPluginBrowser >( "PluginBrowser" );
+	}
+
+	/**
+	 * Checks to see if this module is loaded and ready.  It is only valid to call Get() if IsAvailable() returns true.
+	 *
+	 * @return True if the module is loaded and ready to use
+	 */
+	static inline bool IsAvailable()
+	{
+		return FModuleManager::Get().IsModuleLoaded( "PluginBrowser" );
+	}
+};
+
diff --git a/Engine/Plugins/Editor/PluginBrowser/Source/PluginBrowser/Public/IPluginsEditor.h b/Engine/Plugins/Editor/PluginBrowser/Source/PluginBrowser/Public/IPluginsEditor.h
deleted file mode 100644
diff --git a/Engine/Source/Runtime/Projects/Private/ProjectManager.cpp b/Engine/Source/Runtime/Projects/Private/ProjectManager.cpp
--- a/Engine/Source/Runtime/Projects/Private/ProjectManager.cpp
+++ b/Engine/Source/Runtime/Projects/Private/ProjectManager.cpp
@@ -8,7 +8,6 @@ DEFINE_LOG_CATEGORY_STATIC( LogProjectManager, Log, All );
 
 FProjectManager::FProjectManager()
 {
-	bRestartRequired = false;
 }
 
 const FProjectDescriptor* FProjectManager::GetCurrentProject() const
@@ -323,16 +322,9 @@ bool FProjectManager::SetPluginEnabled(const FString& PluginName, bool bEnabled,
 		return false;
 	}
 
-	// Flag that a restart is required and return
-	bRestartRequired = true;
 	return true;
 }
 
-bool FProjectManager::IsRestartRequired() const
-{
-	return bRestartRequired;
-}
-
 void FProjectManager::GetDefaultEnabledPlugins(TArray<FString>& OutPluginNames)
 {
 	// Add all the game plugins and everything marked as enabled by default
diff --git a/Engine/Source/Runtime/Projects/Private/ProjectManager.h b/Engine/Source/Runtime/Projects/Private/ProjectManager.h
--- a/Engine/Source/Runtime/Projects/Private/ProjectManager.h
+++ b/Engine/Source/Runtime/Projects/Private/ProjectManager.h
@@ -29,7 +29,6 @@ public:
 	virtual void GetEnabledPlugins(TArray<FString>& OutPluginNames) const override;
 	virtual bool IsNonDefaultPluginEnabled() const override;
 	virtual bool SetPluginEnabled(const FString& PluginName, bool bEnabled, FText& OutFailReason) override;
-	virtual bool IsRestartRequired() const override;
 
 private:
 	static void QueryStatusForProjectImpl(const FProjectDescriptor& Project, const FString& FilePath, FProjectStatus& OutProjectStatus);
@@ -42,9 +41,6 @@ private:
 
 	/** Delegate called when the target platforms for the current project are changed */
 	FOnTargetPlatformsForCurrentProjectChangedEvent OnTargetPlatformsForCurrentProjectChangedEvent;
-
-	/** Whether we need to restart in order to reflect change to the current project */
-	bool bRestartRequired;
 };
 
 
diff --git a/Engine/Source/Runtime/Projects/Public/Interfaces/IProjectManager.h b/Engine/Source/Runtime/Projects/Public/Interfaces/IProjectManager.h
--- a/Engine/Source/Runtime/Projects/Public/Interfaces/IProjectManager.h
+++ b/Engine/Source/Runtime/Projects/Public/Interfaces/IProjectManager.h
@@ -201,11 +201,4 @@ public:
 	 * @return	True if the plugin has been marked as enabled, and the project descriptor has been saved.
 	 */
 	virtual bool SetPluginEnabled(const FString& PluginName, bool bEnabled, FText& OutFailReason) = 0;
-
-	/**
-	 * Returns whether a restart is required to reflect changes to the project.
-	 * 
-	 * @return	True if the application needs to be restarted.
-	 */
-	virtual bool IsRestartRequired() const = 0;
 };
\ No newline at end of file
