Change IPluginManager to return shared pointers to plugin objects, rather than passing raw pointers by default.

diff --git a/Engine/Plugins/Editor/PluginBrowser/Source/PluginBrowser/Private/PluginBrowserModule.h b/Engine/Plugins/Editor/PluginBrowser/Source/PluginBrowser/Private/PluginBrowserModule.h
--- a/Engine/Plugins/Editor/PluginBrowser/Source/PluginBrowser/Private/PluginBrowserModule.h
+++ b/Engine/Plugins/Editor/PluginBrowser/Source/PluginBrowser/Private/PluginBrowserModule.h
@@ -9,7 +9,7 @@
 class FPluginBrowserModule : public IPluginBrowser
 {
 public:
-	TMap<IPlugin*, bool> PendingEnablePlugins;
+	TMap<FString, bool> PendingEnablePlugins;
 
 	/** Accessor for the module interface */
 	static FPluginBrowserModule& Get()
diff --git a/Engine/Plugins/Editor/PluginBrowser/Source/PluginBrowser/Private/SPluginCategories.cpp b/Engine/Plugins/Editor/PluginBrowser/Source/PluginBrowser/Private/SPluginCategories.cpp
--- a/Engine/Plugins/Editor/PluginBrowser/Source/PluginBrowser/Private/SPluginCategories.cpp
+++ b/Engine/Plugins/Editor/PluginBrowser/Source/PluginBrowser/Private/SPluginCategories.cpp
@@ -64,7 +64,7 @@ SPluginBrowser& SPluginCategories::GetOwner()
 void SPluginCategories::RebuildAndFilterCategoryTree()
 {
 	// Get a plugin from the currently selected category, so we can track it if it's removed
-	IPlugin* TrackPlugin = nullptr;
+	TSharedPtr<IPlugin> TrackPlugin = nullptr;
 	for(TSharedPtr<FPluginCategory> SelectedItem: PluginCategoryTreeView->GetSelectedItems())
 	{
 		if(SelectedItem->Plugins.Num() > 0)
@@ -85,7 +85,7 @@ void SPluginCategories::RebuildAndFilterCategoryTree()
 
 	// Add all the known plugins into categories
 	TSharedPtr<FPluginCategory> SelectCategory;
-	for(IPlugin* Plugin: IPluginManager::Get().GetDiscoveredPlugins())
+	for(TSharedRef<IPlugin> Plugin: IPluginManager::Get().GetDiscoveredPlugins())
 	{
 		// Figure out which base category this plugin belongs in
 		TSharedPtr<FPluginCategory> RootCategory;
diff --git a/Engine/Plugins/Editor/PluginBrowser/Source/PluginBrowser/Private/SPluginCategoryTreeItem.h b/Engine/Plugins/Editor/PluginBrowser/Source/PluginBrowser/Private/SPluginCategoryTreeItem.h
--- a/Engine/Plugins/Editor/PluginBrowser/Source/PluginBrowser/Private/SPluginCategoryTreeItem.h
+++ b/Engine/Plugins/Editor/PluginBrowser/Source/PluginBrowser/Private/SPluginCategoryTreeItem.h
@@ -21,7 +21,7 @@ public:
 	TArray<TSharedPtr<FPluginCategory>> SubCategories;
 
 	/** Plugins in this category */
-	TArray<IPlugin*> Plugins;
+	TArray<TSharedRef<IPlugin>> Plugins;
 
 	/** Constructor for FPluginCategory */
 	FPluginCategory(TSharedPtr<FPluginCategory> InParentCategory, const FString& InName, const FText& InDisplayName)
diff --git a/Engine/Plugins/Editor/PluginBrowser/Source/PluginBrowser/Private/SPluginList.cpp b/Engine/Plugins/Editor/PluginBrowser/Source/PluginBrowser/Private/SPluginList.cpp
--- a/Engine/Plugins/Editor/PluginBrowser/Source/PluginBrowser/Private/SPluginList.cpp
+++ b/Engine/Plugins/Editor/PluginBrowser/Source/PluginBrowser/Private/SPluginList.cpp
@@ -23,7 +23,7 @@ void SPluginList::Construct( const FArguments& Args, const TSharedRef< SPluginBr
 	// @todo plugedit: Have optional compact version with only plugin icon + name + version?  Only expand selected?
 
 	PluginListView = 
-		SNew( SListView<TSharedRef<IPlugin*>> )
+		SNew( SListView<TSharedRef<IPlugin>> )
 
 		.SelectionMode( ESelectionMode::None )		// No need to select plugins!
 
@@ -52,12 +52,12 @@ SPluginBrowser& SPluginList::GetOwner()
 }
 
 
-TSharedRef<ITableRow> SPluginList::PluginListView_OnGenerateRow( TSharedRef<IPlugin*> Item, const TSharedRef<STableViewBase>& OwnerTable )
+TSharedRef<ITableRow> SPluginList::PluginListView_OnGenerateRow( TSharedRef<IPlugin> Item, const TSharedRef<STableViewBase>& OwnerTable )
 {
 	return
-		SNew( STableRow< TSharedRef<IPlugin*> >, OwnerTable )
+		SNew( STableRow< TSharedRef<IPlugin> >, OwnerTable )
 		[
-			SNew( SPluginListTile, SharedThis( this ), Item.Get() )
+			SNew( SPluginListTile, SharedThis( this ), Item )
 		];
 }
 
@@ -73,11 +73,11 @@ void SPluginList::RebuildAndFilterPluginList()
 		const TSharedPtr<FPluginCategory>& SelectedCategory = OwnerWeak.Pin()->GetSelectedCategory();
 		if( SelectedCategory.IsValid() )
 		{
-			for(IPlugin* Plugin: SelectedCategory->Plugins)
+			for(TSharedRef<IPlugin> Plugin: SelectedCategory->Plugins)
 			{
-				if(OwnerWeak.Pin()->GetPluginTextFilter().PassesFilter(Plugin))
+				if(OwnerWeak.Pin()->GetPluginTextFilter().PassesFilter(&Plugin.Get()))
 				{
-					PluginListItems.Add(TSharedRef<IPlugin*>(new IPlugin*(Plugin)));
+					PluginListItems.Add(Plugin);
 				}
 			}
 		}
@@ -86,9 +86,9 @@ void SPluginList::RebuildAndFilterPluginList()
 		{
 			struct FPluginListItemSorter
 			{
-				bool operator()( const TSharedRef<IPlugin*>& A, const TSharedRef<IPlugin*>& B ) const
+				bool operator()( const TSharedRef<IPlugin>& A, const TSharedRef<IPlugin>& B ) const
 				{
-					return A.Get()->GetDescriptor().FriendlyName < B.Get()->GetDescriptor().FriendlyName;
+					return A->GetDescriptor().FriendlyName < B->GetDescriptor().FriendlyName;
 				}
 			};
 			PluginListItems.Sort( FPluginListItemSorter() );
diff --git a/Engine/Plugins/Editor/PluginBrowser/Source/PluginBrowser/Private/SPluginList.h b/Engine/Plugins/Editor/PluginBrowser/Source/PluginBrowser/Private/SPluginList.h
--- a/Engine/Plugins/Editor/PluginBrowser/Source/PluginBrowser/Private/SPluginList.h
+++ b/Engine/Plugins/Editor/PluginBrowser/Source/PluginBrowser/Private/SPluginList.h
@@ -32,7 +32,7 @@ private:
 	void OnPluginTextFilterChanged();
 
 	/** Called to generate a widget for the specified list item */
-	TSharedRef<ITableRow> PluginListView_OnGenerateRow(TSharedRef<IPlugin*> Item, const TSharedRef<STableViewBase>& OwnerTable );
+	TSharedRef<ITableRow> PluginListView_OnGenerateRow(TSharedRef<IPlugin> Item, const TSharedRef<STableViewBase>& OwnerTable );
 
 	/** Rebuilds the list of plugins from scratch and applies filtering. */
 	void RebuildAndFilterPluginList();
@@ -46,10 +46,10 @@ private:
 	TWeakPtr< class SPluginBrowser > OwnerWeak;
 
 	/** The list view widget for our plugins list */
-	TSharedPtr<SListView<TSharedRef<IPlugin*>>> PluginListView;
+	TSharedPtr<SListView<TSharedRef<IPlugin>>> PluginListView;
 
 	/** List of everything that we want to display in the plugin list */
-	TArray<TSharedRef<IPlugin*>> PluginListItems;
+	TArray<TSharedRef<IPlugin>> PluginListItems;
 
 	/** Whether the active timer to refresh the list is registered */
 	bool bIsActiveTimerRegistered;
diff --git a/Engine/Plugins/Editor/PluginBrowser/Source/PluginBrowser/Private/SPluginListTile.cpp b/Engine/Plugins/Editor/PluginBrowser/Source/PluginBrowser/Private/SPluginListTile.cpp
--- a/Engine/Plugins/Editor/PluginBrowser/Source/PluginBrowser/Private/SPluginListTile.cpp
+++ b/Engine/Plugins/Editor/PluginBrowser/Source/PluginBrowser/Private/SPluginListTile.cpp
@@ -16,7 +16,7 @@
 #define LOCTEXT_NAMESPACE "PluginListTile"
 
 
-void SPluginListTile::Construct( const FArguments& Args, const TSharedRef<SPluginList> Owner, IPlugin* InPlugin )
+void SPluginListTile::Construct( const FArguments& Args, const TSharedRef<SPluginList> Owner, TSharedRef<IPlugin> InPlugin )
 {
 	OwnerWeak = Owner;
 	Plugin = InPlugin;
@@ -310,9 +310,9 @@ void SPluginListTile::Construct( const FArguments& Args, const TSharedRef<SPlugi
 ECheckBoxState SPluginListTile::IsPluginEnabled() const
 {
 	FPluginBrowserModule& PluginBrowserModule = FPluginBrowserModule::Get();
-	if(PluginBrowserModule.PendingEnablePlugins.Contains(Plugin))
+	if(PluginBrowserModule.PendingEnablePlugins.Contains(Plugin->GetName()))
 	{
-		return PluginBrowserModule.PendingEnablePlugins[Plugin]? ECheckBoxState::Checked : ECheckBoxState::Unchecked;;
+		return PluginBrowserModule.PendingEnablePlugins[Plugin->GetName()]? ECheckBoxState::Checked : ECheckBoxState::Unchecked;;
 	}
 	else
 	{
@@ -350,11 +350,11 @@ void SPluginListTile::OnEnablePluginCheckboxChanged(ECheckBoxState NewCheckedSta
 		FPluginBrowserModule& PluginBrowserModule = FPluginBrowserModule::Get();
 		if(Plugin->IsEnabled() == bNewEnabledState)
 		{
-			PluginBrowserModule.PendingEnablePlugins.Remove(Plugin);
+			PluginBrowserModule.PendingEnablePlugins.Remove(Plugin->GetName());
 		}
 		else
 		{
-			PluginBrowserModule.PendingEnablePlugins.FindOrAdd(Plugin) = bNewEnabledState;
+			PluginBrowserModule.PendingEnablePlugins.FindOrAdd(Plugin->GetName()) = bNewEnabledState;
 		}
 	}
 
diff --git a/Engine/Plugins/Editor/PluginBrowser/Source/PluginBrowser/Private/SPluginListTile.h b/Engine/Plugins/Editor/PluginBrowser/Source/PluginBrowser/Private/SPluginListTile.h
--- a/Engine/Plugins/Editor/PluginBrowser/Source/PluginBrowser/Private/SPluginListTile.h
+++ b/Engine/Plugins/Editor/PluginBrowser/Source/PluginBrowser/Private/SPluginListTile.h
@@ -15,7 +15,7 @@ class SPluginListTile : public SCompoundWidget
 
 
 	/** Widget constructor */
-	void Construct( const FArguments& Args, const TSharedRef< class SPluginList > Owner, IPlugin* Plugin );
+	void Construct( const FArguments& Args, const TSharedRef< class SPluginList > Owner, TSharedRef<IPlugin> Plugin );
 
 
 private:
@@ -38,7 +38,7 @@ private:
 private:
 
 	/** The item we're representing the in tree */
-	IPlugin* Plugin;
+	TSharedPtr<IPlugin> Plugin;
 
 	/** Weak pointer back to its owner */
 	TWeakPtr< class SPluginList > OwnerWeak;
diff --git a/Engine/Source/Editor/ContentBrowser/Private/ContentBrowserUtils.cpp b/Engine/Source/Editor/ContentBrowser/Private/ContentBrowserUtils.cpp
--- a/Engine/Source/Editor/ContentBrowser/Private/ContentBrowserUtils.cpp
+++ b/Engine/Source/Editor/ContentBrowser/Private/ContentBrowserUtils.cpp
@@ -1097,7 +1097,7 @@ bool ContentBrowserUtils::IsDevelopersFolder( const FString& InPath )
 bool ContentBrowserUtils::IsPluginFolder( const FString& InPath )
 {
 	FString PathWithSlash = InPath / TEXT("");
-	for(const IPlugin* Plugin: IPluginManager::Get().GetEnabledPlugins())
+	for(const TSharedRef<IPlugin>& Plugin: IPluginManager::Get().GetEnabledPlugins())
 	{
 		if(Plugin->CanContainContent())
 		{
diff --git a/Engine/Source/Runtime/Projects/Private/PluginManager.cpp b/Engine/Source/Runtime/Projects/Private/PluginManager.cpp
--- a/Engine/Source/Runtime/Projects/Private/PluginManager.cpp
+++ b/Engine/Source/Runtime/Projects/Private/PluginManager.cpp
@@ -411,7 +411,7 @@ bool FPluginManager::ConfigureEnabledPlugins()
 		IPlatformFile& PlatformFile = FPlatformFileManager::Get().GetPlatformFile();
 		if( ensure( RegisterMountPointDelegate.IsBound() ) )
 		{
-			for(IPlugin* Plugin: GetEnabledPlugins())
+			for(TSharedRef<IPlugin> Plugin: GetEnabledPlugins())
 			{
 				if(Plugin->CanContainContent())
 				{
@@ -560,39 +560,39 @@ IPluginManager& IPluginManager::Get()
 	return *PluginManager;
 }
 
-FPlugin* FPluginManager::FindPlugin(const FString& Name)
+TSharedPtr<IPlugin> FPluginManager::FindPlugin(const FString& Name)
 {
-	FPlugin* Plugin = nullptr;
+	TSharedPtr<IPlugin> Plugin;
 	for(TSharedRef<FPlugin>& PossiblePlugin : AllPlugins)
 	{
 		if(PossiblePlugin->Name == Name)
 		{
-			Plugin = &(PossiblePlugin.Get());
+			Plugin = PossiblePlugin;
 			break;
 		}
 	}
 	return Plugin;
 }
 
-TArray<IPlugin*> FPluginManager::GetEnabledPlugins()
+TArray<TSharedRef<IPlugin>> FPluginManager::GetEnabledPlugins()
 {
-	TArray<IPlugin*> Plugins;
+	TArray<TSharedRef<IPlugin>> Plugins;
 	for(TSharedRef<FPlugin>& PossiblePlugin : AllPlugins)
 	{
 		if(PossiblePlugin->bEnabled)
 		{
-			Plugins.Add(&(PossiblePlugin.Get()));
+			Plugins.Add(PossiblePlugin);
 		}
 	}
 	return Plugins;
 }
 
-TArray<IPlugin*> FPluginManager::GetDiscoveredPlugins()
+TArray<TSharedRef<IPlugin>> FPluginManager::GetDiscoveredPlugins()
 {
-	TArray<IPlugin*> Plugins;
+	TArray<TSharedRef<IPlugin>> Plugins;
 	for(TSharedRef<FPlugin>& Plugin : AllPlugins)
 	{
-		Plugins.Add(&(Plugin.Get()));
+		Plugins.Add(Plugin);
 	}
 	return Plugins;
 }
diff --git a/Engine/Source/Runtime/Projects/Private/PluginManager.h b/Engine/Source/Runtime/Projects/Private/PluginManager.h
--- a/Engine/Source/Runtime/Projects/Private/PluginManager.h
+++ b/Engine/Source/Runtime/Projects/Private/PluginManager.h
@@ -65,9 +65,9 @@ public:
 	virtual void SetRegisterMountPointDelegate( const FRegisterMountPointDelegate& Delegate ) override;
 	virtual bool AreRequiredPluginsAvailable() override;
 	virtual bool CheckModuleCompatibility( TArray<FString>& OutIncompatibleModules ) override;
-	virtual FPlugin* FindPlugin(const FString& Name) override;
-	virtual TArray<IPlugin*> GetEnabledPlugins() override;
-	virtual TArray<IPlugin*> GetDiscoveredPlugins() override;
+	virtual TSharedPtr<IPlugin> FindPlugin(const FString& Name) override;
+	virtual TArray<TSharedRef<IPlugin>> GetEnabledPlugins() override;
+	virtual TArray<TSharedRef<IPlugin>> GetDiscoveredPlugins() override;
 	virtual TArray< FPluginStatus > QueryStatusForAllPlugins() const override;
 
 private:
diff --git a/Engine/Source/Runtime/Projects/Public/Interfaces/IPluginManager.h b/Engine/Source/Runtime/Projects/Public/Interfaces/IPluginManager.h
--- a/Engine/Source/Runtime/Projects/Public/Interfaces/IPluginManager.h
+++ b/Engine/Source/Runtime/Projects/Public/Interfaces/IPluginManager.h
@@ -46,6 +46,9 @@ struct FPluginStatus
 class IPlugin
 {
 public:
+	/* Virtual destructor */
+	virtual ~IPlugin(){}
+
 	/**
 	 * Gets the plugin name.
 	 *
@@ -165,21 +168,21 @@ public:
 	 *
 	 * @return	 Pointer to the plugin's information, or nullptr.
 	 */
-	virtual IPlugin* FindPlugin(const FString& Name) = 0;
+	virtual TSharedPtr<IPlugin> FindPlugin(const FString& Name) = 0;
 
 	/**
 	 * Gets an array of all the enabled plugins.
 	 *
 	 * @return	Array of the enabled plugins.
 	 */
-	virtual TArray<IPlugin*> GetEnabledPlugins() = 0;
+	virtual TArray<TSharedRef<IPlugin>> GetEnabledPlugins() = 0;
 
 	/**
 	 * Gets an array of all the discovered plugins.
 	 *
 	 * @return	Array of the discovered plugins.
 	 */
-	virtual TArray<IPlugin*> GetDiscoveredPlugins() = 0;
+	virtual TArray<TSharedRef<IPlugin>> GetDiscoveredPlugins() = 0;
 
 	/**
 	 * Gets status about all currently known plug-ins.
