Fix tile widgets not updating when modifying plugin settings.

diff --git a/Engine/Plugins/Editor/PluginBrowser/Source/PluginBrowser/Private/SPluginListTile.cpp b/Engine/Plugins/Editor/PluginBrowser/Source/PluginBrowser/Private/SPluginListTile.cpp
--- a/Engine/Plugins/Editor/PluginBrowser/Source/PluginBrowser/Private/SPluginListTile.cpp
+++ b/Engine/Plugins/Editor/PluginBrowser/Source/PluginBrowser/Private/SPluginListTile.cpp
@@ -21,6 +21,11 @@ void SPluginListTile::Construct( const FArguments& Args, const TSharedRef<SPlugi
 	OwnerWeak = Owner;
 	Plugin = InPlugin;
 
+	RecreateWidgets();
+}
+
+void SPluginListTile::RecreateWidgets()
+{
 	const float PaddingAmount = FPluginStyle::Get()->GetFloat( "PluginTile.Padding" );
 	const float ThumbnailImageSize = FPluginStyle::Get()->GetFloat( "PluginTile.ThumbnailImageSize" );
 
@@ -462,7 +467,10 @@ FReply SPluginListTile::OnEditPluginFinished(UPluginMetadataObject* MetadataObje
 			FMessageDialog::Open(EAppMsgType::Ok, FailReason);
 		}
 
-		// Refresh the view
+		// Recreate the widgets on this tile
+		RecreateWidgets();
+
+		// Refresh the parent too
 		if(OwnerWeak.IsValid())
 		{
 			OwnerWeak.Pin()->GetOwner().SetNeedsRefresh();
diff --git a/Engine/Plugins/Editor/PluginBrowser/Source/PluginBrowser/Private/SPluginListTile.h b/Engine/Plugins/Editor/PluginBrowser/Source/PluginBrowser/Private/SPluginListTile.h
--- a/Engine/Plugins/Editor/PluginBrowser/Source/PluginBrowser/Private/SPluginListTile.h
+++ b/Engine/Plugins/Editor/PluginBrowser/Source/PluginBrowser/Private/SPluginListTile.h
@@ -17,9 +17,11 @@ class SPluginListTile : public SCompoundWidget
 	/** Widget constructor */
 	void Construct( const FArguments& Args, const TSharedRef< class SPluginList > Owner, TSharedRef<IPlugin> Plugin );
 
-
 private:
 
+	/** Updates the contents of this tile */
+	void RecreateWidgets();
+
 	/** Returns the checked state for the enabled checkbox */
 	ECheckBoxState IsPluginEnabled() const;
 
