Fix some more things in PluginBrowser for consistency, and remove redundant wrappers around categories in breadcrumb bar.

diff --git a/Engine/Plugins/Editor/PluginBrowser/Source/PluginBrowser/Private/SPluginBrowser.cpp b/Engine/Plugins/Editor/PluginBrowser/Source/PluginBrowser/Private/SPluginBrowser.cpp
--- a/Engine/Plugins/Editor/PluginBrowser/Source/PluginBrowser/Private/SPluginBrowser.cpp
+++ b/Engine/Plugins/Editor/PluginBrowser/Source/PluginBrowser/Private/SPluginBrowser.cpp
@@ -3,8 +3,8 @@
 #include "PluginBrowserPrivatePCH.h"
 #include "PluginStyle.h"
 #include "SPluginBrowser.h"
-#include "SPluginList.h"
-#include "SPluginCategories.h"
+#include "SPluginTileList.h"
+#include "SPluginCategoryTree.h"
 #include "SSearchBox.h"
 
 #define LOCTEXT_NAMESPACE "PluginsEditor"
@@ -34,7 +34,7 @@ void SPluginBrowser::Construct( const FArguments& Args )
 	const float PaddingAmount = 2.0f;
 
 
-	PluginCategories = SNew( SPluginCategories, SharedThis( this ) );
+	PluginCategories = SNew( SPluginCategoryTree, SharedThis( this ) );
 
 
 	ChildSlot
@@ -62,7 +62,7 @@ void SPluginBrowser::Construct( const FArguments& Args )
 				+SHorizontalBox::Slot()
 				.Padding( PaddingAmount )
 				[
-					SAssignNew( BreadcrumbTrail, SPluginCategoryBreadcrumbTrail )
+					SAssignNew( BreadcrumbTrail, SBreadcrumbTrail<TSharedPtr<FPluginCategory>> )
 					.DelimiterImage( FPluginStyle::Get()->GetBrush( "Plugins.BreadcrumbArrow" ) ) 
 					.ShowLeadingDelimiter( true )
 					.OnCrumbClicked( this, &SPluginBrowser::BreadcrumbTrail_OnCrumbClicked )
@@ -78,7 +78,7 @@ void SPluginBrowser::Construct( const FArguments& Args )
 
 			+SVerticalBox::Slot()
 			[
-				SAssignNew( PluginList, SPluginList, SharedThis( this ) )
+				SAssignNew( PluginList, SPluginTileList, SharedThis( this ) )
 			]
 
 			+SVerticalBox::Slot()
@@ -205,18 +205,17 @@ void SPluginBrowser::RefreshBreadcrumbTrail()
 		BreadcrumbTrail->ClearCrumbs();
 		for(TSharedPtr<FPluginCategory>& Category: CategoryPath)
 		{
-			FPluginCategoryBreadcrumbPtr NewBreadcrumb( new FPluginCategoryBreadcrumb( Category ) );
-			BreadcrumbTrail->PushCrumb( Category->DisplayName, NewBreadcrumb );
+			BreadcrumbTrail->PushCrumb( Category->DisplayName, Category );
 		}
 	}
 }
 
 
-void SPluginBrowser::BreadcrumbTrail_OnCrumbClicked( const FPluginCategoryBreadcrumbPtr& CrumbData )
+void SPluginBrowser::BreadcrumbTrail_OnCrumbClicked( const TSharedPtr<FPluginCategory>& Category )
 {
 	if( PluginCategories.IsValid() )
 	{
-		PluginCategories->SelectCategory( CrumbData->GetCategoryItem() );
+		PluginCategories->SelectCategory( Category );
 	}
 }
 
diff --git a/Engine/Plugins/Editor/PluginBrowser/Source/PluginBrowser/Private/SPluginBrowser.h b/Engine/Plugins/Editor/PluginBrowser/Source/PluginBrowser/Private/SPluginBrowser.h
--- a/Engine/Plugins/Editor/PluginBrowser/Source/PluginBrowser/Private/SPluginBrowser.h
+++ b/Engine/Plugins/Editor/PluginBrowser/Source/PluginBrowser/Private/SPluginBrowser.h
@@ -3,39 +3,7 @@
 #pragma once
 
 #include "SBreadcrumbTrail.h"
-
-/**
- * Stores data about a breadcrumb entry for the plugins editor
- */
-class FPluginCategoryBreadcrumb
-{
-
-public:
-
-	/** Initializing constructor */
-	FPluginCategoryBreadcrumb( TSharedPtr< class FPluginCategory > InitCategoryItem )
-		: CategoryItem( InitCategoryItem )
-	{
-	}
-
-	/** @return Gets the category item for this breadcrumb */
-	const TSharedPtr< class FPluginCategory >& GetCategoryItem() const
-	{
-		return CategoryItem;
-	}
-
-
-private:
-
-	/** The category tree item data for this breadcrumb */
-	TSharedPtr< class FPluginCategory > CategoryItem;
-};
-
-
-
-typedef TSharedPtr< FPluginCategoryBreadcrumb > FPluginCategoryBreadcrumbPtr;
-typedef SBreadcrumbTrail< FPluginCategoryBreadcrumbPtr > SPluginCategoryBreadcrumbTrail;
-
+#include "SPluginCategory.h"
 
 /**
  * Implementation of main plugin editor Slate widget
@@ -79,7 +47,7 @@ private:
 	void SearchBox_OnPluginSearchTextChanged( const FText& NewText );
 
 	/** Called when a breadcrumb is clicked on the breadcrumb trail */
-	void BreadcrumbTrail_OnCrumbClicked( const FPluginCategoryBreadcrumbPtr& CrumbData );
+	void BreadcrumbTrail_OnCrumbClicked( const TSharedPtr<FPluginCategory>& CrumbData );
 
 	/** Called to refresh the breadcrumb trail immediately */
 	void RefreshBreadcrumbTrail();
@@ -90,15 +58,15 @@ private:
 private:
 
 	/** The plugin categories widget */
-	TSharedPtr< class SPluginCategories > PluginCategories;
+	TSharedPtr< class SPluginCategoryTree > PluginCategories;
 
 	/** The plugin list widget */
-	TSharedPtr< class SPluginList > PluginList;
+	TSharedPtr< class SPluginTileList > PluginList;
 
 	/** Text filter object for typing in filter text to the search box */
 	TSharedPtr< FPluginTextFilter > PluginTextFilter;
 
 	/** Breadcrumb trail widget for the currently selected category */
-	TSharedPtr< SPluginCategoryBreadcrumbTrail > BreadcrumbTrail;
+	TSharedPtr< SBreadcrumbTrail< TSharedPtr< FPluginCategory > > > BreadcrumbTrail;
 };
 
diff --git a/Engine/Plugins/Editor/PluginBrowser/Source/PluginBrowser/Private/SPluginCategories.cpp b/Engine/Plugins/Editor/PluginBrowser/Source/PluginBrowser/Private/SPluginCategories.cpp
deleted file mode 100644
diff --git a/Engine/Plugins/Editor/PluginBrowser/Source/PluginBrowser/Private/SPluginCategories.h b/Engine/Plugins/Editor/PluginBrowser/Source/PluginBrowser/Private/SPluginCategories.h
deleted file mode 100644
diff --git a/Engine/Plugins/Editor/PluginBrowser/Source/PluginBrowser/Private/SPluginCategory.cpp b/Engine/Plugins/Editor/PluginBrowser/Source/PluginBrowser/Private/SPluginCategory.cpp
new file mode 100644
--- /dev/null
+++ b/Engine/Plugins/Editor/PluginBrowser/Source/PluginBrowser/Private/SPluginCategory.cpp
@@ -0,0 +1,90 @@
+// Copyright 1998-2015 Epic Games, Inc. All Rights Reserved.
+
+#include "PluginBrowserPrivatePCH.h"
+#include "SPluginCategory.h"
+#include "PluginStyle.h"
+
+
+#define LOCTEXT_NAMESPACE "PluginCategoryTreeItem"
+
+
+void SPluginCategory::Construct(const FArguments& Args, const TSharedRef<FPluginCategory>& InCategory)
+{
+	Category = InCategory;
+
+	const float CategoryIconSize = FPluginStyle::Get()->GetFloat("CategoryTreeItem.IconSize");
+	const float PaddingAmount = FPluginStyle::Get()->GetFloat("CategoryTreeItem.PaddingAmount");;
+
+	// Figure out which font size to use
+	const auto bIsRootItem = !Category->ParentCategory.IsValid();
+
+	ChildSlot
+	[ 
+		SNew( SBorder )
+		.BorderImage( FPluginStyle::Get()->GetBrush(bIsRootItem ? "CategoryTreeItem.Root.BackgroundBrush" : "CategoryTreeItem.BackgroundBrush") )
+		.Padding(FPluginStyle::Get()->GetMargin(bIsRootItem ? "CategoryTreeItem.Root.BackgroundPadding" : "CategoryTreeItem.BackgroundPadding") )
+		[
+			SNew( SHorizontalBox )
+
+			// Icon image
+			+SHorizontalBox::Slot()
+			.Padding( PaddingAmount )
+			.AutoWidth()
+			.VAlign(VAlign_Center)
+			[
+				SNew( SBox )
+				.WidthOverride( CategoryIconSize )
+				.HeightOverride( CategoryIconSize )
+				[
+					SNew( SImage )
+					.Image( this, &SPluginCategory::GetIconBrush )
+				]
+			]
+
+			// Category name
+			+SHorizontalBox::Slot()
+			.Padding( PaddingAmount )
+			.VAlign(VAlign_Center)
+			[
+				SNew( STextBlock )
+				.Text( Category->DisplayName )
+				.TextStyle( FPluginStyle::Get(), bIsRootItem ? "CategoryTreeItem.Root.Text" : "CategoryTreeItem.Text" )
+			]
+			
+			// Plugin count
+			+SHorizontalBox::Slot()
+			.AutoWidth()
+			.Padding( PaddingAmount )
+			.VAlign(VAlign_Center)
+			[
+				SNew( STextBlock )
+
+				// Only display if at there is least one plugin is in this category
+				.Visibility( Category->Plugins.Num() > 0 ? EVisibility::Visible : EVisibility::Collapsed )
+
+				.Text( FText::Format( LOCTEXT( "NumberOfPluginsWrapper", "({0})" ), FText::AsNumber( Category->Plugins.Num() ) ) )
+				.TextStyle( FPluginStyle::Get(), bIsRootItem ? "CategoryTreeItem.Root.PluginCountText" : "CategoryTreeItem.PluginCountText" )
+			]
+		]
+	];
+}
+
+
+const FSlateBrush* SPluginCategory::GetIconBrush() const
+{
+	if(Category->ParentCategory.IsValid())
+	{
+		return FPluginStyle::Get()->GetBrush( "CategoryTreeItem.LeafItemWithPlugin" );
+	}
+	else if (Category->Name == TEXT("Installed"))
+	{
+		return FPluginStyle::Get()->GetBrush( "CategoryTreeItem.BuiltIn" );
+	}
+	else
+	{
+		return FPluginStyle::Get()->GetBrush( "CategoryTreeItem.Installed" );
+	}
+}
+
+
+#undef LOCTEXT_NAMESPACE
\ No newline at end of file
diff --git a/Engine/Plugins/Editor/PluginBrowser/Source/PluginBrowser/Private/SPluginCategory.h b/Engine/Plugins/Editor/PluginBrowser/Source/PluginBrowser/Private/SPluginCategory.h
new file mode 100644
--- /dev/null
+++ b/Engine/Plugins/Editor/PluginBrowser/Source/PluginBrowser/Private/SPluginCategory.h
@@ -0,0 +1,59 @@
+// Copyright 1998-2015 Epic Games, Inc. All Rights Reserved.
+
+#pragma once
+
+/**
+ * Represents a category in the plugin category tree
+ */
+class FPluginCategory
+{
+public:
+	/** Parent category item or NULL if this is a root category */
+	TWeakPtr<FPluginCategory> ParentCategory;
+
+	/** Name of the category */
+	FString Name;
+
+	/** Display name of the category */
+	FText DisplayName;
+
+	/** Child categories */
+	TArray<TSharedPtr<FPluginCategory>> SubCategories;
+
+	/** Plugins in this category */
+	TArray<TSharedRef<IPlugin>> Plugins;
+
+	/** Constructor for FPluginCategory */
+	FPluginCategory(TSharedPtr<FPluginCategory> InParentCategory, const FString& InName, const FText& InDisplayName)
+		: ParentCategory(InParentCategory),
+		  Name(InName),
+		  DisplayName(InDisplayName)
+	{
+	}
+};
+
+/**
+ * Widget that represents a single category in the category tree view
+ */
+class SPluginCategory : public SCompoundWidget
+{
+	SLATE_BEGIN_ARGS( SPluginCategory )
+	{
+	}
+
+	SLATE_END_ARGS()
+
+	/** Widget constructor */
+	void Construct(const FArguments& Args, const TSharedRef<FPluginCategory>& Item);
+
+private:
+
+	/** @return Gets the icon brush to use for this item's current state */
+	const FSlateBrush* GetIconBrush() const;
+
+private:
+
+	/** The item we're representing the in tree */
+	TSharedPtr<FPluginCategory> Category;
+};
+
diff --git a/Engine/Plugins/Editor/PluginBrowser/Source/PluginBrowser/Private/SPluginCategoryTree.cpp b/Engine/Plugins/Editor/PluginBrowser/Source/PluginBrowser/Private/SPluginCategoryTree.cpp
new file mode 100644
--- /dev/null
+++ b/Engine/Plugins/Editor/PluginBrowser/Source/PluginBrowser/Private/SPluginCategoryTree.cpp
@@ -0,0 +1,234 @@
+// Copyright 1998-2015 Epic Games, Inc. All Rights Reserved.
+
+#include "PluginBrowserPrivatePCH.h"
+#include "SPluginCategoryTree.h"
+#include "IPluginManager.h"
+#include "SPluginCategory.h"
+#include "SPluginBrowser.h"
+
+#define LOCTEXT_NAMESPACE "PluginCategories"
+
+
+
+void SPluginCategoryTree::Construct( const FArguments& Args, const TSharedRef< SPluginBrowser > Owner )
+{
+	OwnerWeak = Owner;
+
+	RootPluginCategories.Reset();
+
+	// Add the root level categories
+	TSharedRef<FPluginCategory> InstalledCategory = MakeShareable(new FPluginCategory(NULL, TEXT("Installed"), LOCTEXT("InstalledCategoryName", "Installed")));
+	RootPluginCategories.Add( InstalledCategory );
+	TSharedRef<FPluginCategory> ProjectCategory = MakeShareable(new FPluginCategory(NULL, TEXT("Project"), LOCTEXT("ProjectCategoryName", "Project")));
+	RootPluginCategories.Add( ProjectCategory );
+
+	PluginCategoryTreeView =
+		SNew( STreeView<TSharedPtr<FPluginCategory>> )
+
+		// For now we only support selecting a single folder in the tree
+		.SelectionMode( ESelectionMode::Single )
+		.ClearSelectionOnClick( false )		// Don't allow user to select nothing.  We always expect a category to be selected!
+
+		.TreeItemsSource( &RootPluginCategories )
+		.OnGenerateRow( this, &SPluginCategoryTree::PluginCategoryTreeView_OnGenerateRow ) 
+		.OnGetChildren( this, &SPluginCategoryTree::PluginCategoryTreeView_OnGetChildren )
+
+		.OnSelectionChanged( this, &SPluginCategoryTree::PluginCategoryTreeView_OnSelectionChanged )
+		;
+
+	// Expand the root categories by default
+	for( auto RootCategoryIt( RootPluginCategories.CreateConstIterator() ); RootCategoryIt; ++RootCategoryIt )
+	{
+		const auto& Category = *RootCategoryIt;
+		PluginCategoryTreeView->SetItemExpansion( Category, true );
+	}
+
+	RebuildAndFilterCategoryTree();
+
+	ChildSlot.AttachWidget( PluginCategoryTreeView.ToSharedRef() );
+}
+
+
+SPluginCategoryTree::~SPluginCategoryTree()
+{
+}
+
+
+/** @return Gets the owner of this list */
+SPluginBrowser& SPluginCategoryTree::GetOwner()
+{
+	return *OwnerWeak.Pin();
+}
+
+
+void SPluginCategoryTree::RebuildAndFilterCategoryTree()
+{
+	// Get a plugin from the currently selected category, so we can track it if it's removed
+	TSharedPtr<IPlugin> TrackPlugin = nullptr;
+	for(TSharedPtr<FPluginCategory> SelectedItem: PluginCategoryTreeView->GetSelectedItems())
+	{
+		if(SelectedItem->Plugins.Num() > 0)
+		{
+			TrackPlugin = SelectedItem->Plugins[0];
+			break;
+		}
+	}
+
+	// Clear the list of plugins in each current category
+	for(TSharedPtr<FPluginCategory> RootCategory: RootPluginCategories)
+	{
+		for(TSharedPtr<FPluginCategory> Category: RootCategory->SubCategories)
+		{
+			Category->Plugins.Reset();
+		}
+	}
+
+	// Add all the known plugins into categories
+	TSharedPtr<FPluginCategory> SelectCategory;
+	for(TSharedRef<IPlugin> Plugin: IPluginManager::Get().GetDiscoveredPlugins())
+	{
+		// Figure out which base category this plugin belongs in
+		TSharedPtr<FPluginCategory> RootCategory;
+		if( Plugin->GetLoadedFrom() == EPluginLoadedFrom::Engine )
+		{
+			RootCategory = RootPluginCategories[0];
+		}
+		else
+		{
+			RootCategory = RootPluginCategories[1];
+		}
+
+		// Get the subcategory for this plugin
+		FString CategoryName = Plugin->GetDescriptor().Category;
+		if(CategoryName.IsEmpty())
+		{
+			CategoryName = TEXT("Other");
+		}
+
+		// Locate this category at the level we're at in the hierarchy
+		TSharedPtr<FPluginCategory> FoundCategory = NULL;
+		for(TSharedPtr<FPluginCategory> TestCategory: RootCategory->SubCategories)
+		{
+			if(TestCategory->Name == CategoryName)
+			{
+				FoundCategory = TestCategory;
+				break;
+			}
+		}
+
+		if( !FoundCategory.IsValid() )
+		{
+			//@todo Allow for properly localized category names [3/7/2014 justin.sargent]
+			FoundCategory = MakeShareable(new FPluginCategory(RootCategory, CategoryName, FText::FromString(CategoryName)));
+			RootCategory->SubCategories.Add( FoundCategory );
+		}
+			
+		// Associate the plugin with the category
+		FoundCategory->Plugins.Add(Plugin);
+
+		// Update the selection if this is the plugin we're tracking
+		if(TrackPlugin == Plugin)
+		{
+			SelectCategory = FoundCategory;
+		}
+	}
+
+	// Remove any empty categories, keeping track of which items are still selected
+	for(TSharedPtr<FPluginCategory> RootCategory: RootPluginCategories)
+	{
+		for(int32 Idx = 0; Idx < RootCategory->SubCategories.Num(); Idx++)
+		{
+			if(RootCategory->SubCategories[Idx]->Plugins.Num() == 0)
+			{
+				RootCategory->SubCategories.RemoveAt(Idx);
+			}
+		}
+	}
+
+	// Sort every single category alphabetically
+	for(TSharedPtr<FPluginCategory> RootCategory: RootPluginCategories)
+	{
+		RootCategory->SubCategories.Sort([](const TSharedPtr<FPluginCategory>& A, const TSharedPtr<FPluginCategory>& B) -> bool { return A->DisplayName.CompareTo(B->DisplayName) < 0; });
+	}
+
+	// Refresh the view
+	PluginCategoryTreeView->RequestTreeRefresh();
+
+	// Make sure we have something selected
+	if(SelectCategory.IsValid())
+	{
+		PluginCategoryTreeView->SetSelection(SelectCategory);
+	}
+	else if(RootPluginCategories[0]->SubCategories.Num() > 0)
+	{
+		PluginCategoryTreeView->SetSelection(RootPluginCategories[0]->SubCategories[0]);
+	}
+	else
+	{
+		PluginCategoryTreeView->SetSelection(RootPluginCategories[0]);
+	}
+}
+
+TSharedRef<ITableRow> SPluginCategoryTree::PluginCategoryTreeView_OnGenerateRow( TSharedPtr<FPluginCategory> Item, const TSharedRef<STableViewBase>& OwnerTable )
+{
+	return
+		SNew( STableRow<TSharedPtr<FPluginCategory>>, OwnerTable )
+		[
+			SNew( SPluginCategory, Item.ToSharedRef() )
+		];
+}
+
+
+void SPluginCategoryTree::PluginCategoryTreeView_OnGetChildren(TSharedPtr<FPluginCategory> Item, TArray<TSharedPtr<FPluginCategory>>& OutChildren )
+{
+	OutChildren.Append(Item->SubCategories);
+}
+
+
+void SPluginCategoryTree::PluginCategoryTreeView_OnSelectionChanged(TSharedPtr<FPluginCategory> Item, ESelectInfo::Type SelectInfo )
+{
+	// Selection changed, which may affect which plugins are displayed in the list.  We need to invalidate the list.
+	OwnerWeak.Pin()->OnCategorySelectionChanged();
+}
+
+
+TSharedPtr<FPluginCategory> SPluginCategoryTree::GetSelectedCategory() const
+{
+	if( PluginCategoryTreeView.IsValid() )
+	{
+		auto SelectedItems = PluginCategoryTreeView->GetSelectedItems();
+		if( SelectedItems.Num() > 0 )
+		{
+			const auto& SelectedCategoryItem = SelectedItems[ 0 ];
+			return SelectedCategoryItem;
+		}
+	}
+
+	return NULL;
+}
+
+void SPluginCategoryTree::SelectCategory( const TSharedPtr<FPluginCategory>& CategoryToSelect )
+{
+	if( ensure( PluginCategoryTreeView.IsValid() ) )
+	{
+		PluginCategoryTreeView->SetSelection( CategoryToSelect );
+	}
+}
+
+bool SPluginCategoryTree::IsItemExpanded( const TSharedPtr<FPluginCategory> Item ) const
+{
+	return PluginCategoryTreeView->IsItemExpanded( Item );
+}
+
+void SPluginCategoryTree::SetNeedsRefresh()
+{
+	RegisterActiveTimer (0.f, FWidgetActiveTimerDelegate::CreateSP (this, &SPluginCategoryTree::TriggerCategoriesRefresh));
+}
+
+EActiveTimerReturnType SPluginCategoryTree::TriggerCategoriesRefresh(double InCurrentTime, float InDeltaTime)
+{
+	RebuildAndFilterCategoryTree();
+	return EActiveTimerReturnType::Stop;
+}
+
+#undef LOCTEXT_NAMESPACE
\ No newline at end of file
diff --git a/Engine/Plugins/Editor/PluginBrowser/Source/PluginBrowser/Private/SPluginCategoryTree.h b/Engine/Plugins/Editor/PluginBrowser/Source/PluginBrowser/Private/SPluginCategoryTree.h
new file mode 100644
--- /dev/null
+++ b/Engine/Plugins/Editor/PluginBrowser/Source/PluginBrowser/Private/SPluginCategoryTree.h
@@ -0,0 +1,71 @@
+// Copyright 1998-2015 Epic Games, Inc. All Rights Reserved.
+
+#pragma once
+
+#include "SPluginCategory.h"
+
+/**
+ * Tree view that displays all of the plugin categories and allows the user to switch views
+ */
+class SPluginCategoryTree : public SCompoundWidget
+{
+
+public:
+
+	SLATE_BEGIN_ARGS( SPluginCategoryTree )
+	{
+	}
+
+	SLATE_END_ARGS()
+
+
+	/** Widget constructor */
+	void Construct( const FArguments& Args, const TSharedRef< class SPluginBrowser > Owner );
+
+	/** Destructor */
+	~SPluginCategoryTree();
+
+	/** @return Gets the owner of this categories tree */
+	SPluginBrowser& GetOwner();
+
+	/** @return Returns the currently selected category item */
+	TSharedPtr<FPluginCategory> GetSelectedCategory() const;
+
+	/** Selects the specified category */
+	void SelectCategory(const TSharedPtr<FPluginCategory>& CategoryToSelect);
+
+	/** @return Returns true if the specified item is currently expanded in the tree */
+	bool IsItemExpanded(const TSharedPtr<FPluginCategory> Item) const;
+
+	/** Signal that the categories list needs to be refreshed */
+	void SetNeedsRefresh();
+
+private:
+	
+	/** Called to generate a widget for the specified tree item */
+	TSharedRef<ITableRow> PluginCategoryTreeView_OnGenerateRow(TSharedPtr<FPluginCategory> Item, const TSharedRef<STableViewBase>& OwnerTable);
+
+	/** Given a tree item, fills an array of child items */
+	void PluginCategoryTreeView_OnGetChildren(TSharedPtr<FPluginCategory> Item, TArray<TSharedPtr<FPluginCategory>>& OutChildren);
+
+	/** Called when the user clicks on a plugin item, or when selection changes by some other means */
+	void PluginCategoryTreeView_OnSelectionChanged(TSharedPtr<FPluginCategory> Item, ESelectInfo::Type SelectInfo);
+
+	/** Rebuilds the category tree from scratch */
+	void RebuildAndFilterCategoryTree();
+
+	/** Callback for refreshing the tree */
+	EActiveTimerReturnType TriggerCategoriesRefresh(double InCurrentTime, float InDeltaTime);
+
+private:
+
+	/** Weak pointer back to its owner */
+	TWeakPtr< class SPluginBrowser > OwnerWeak;
+
+	/** The tree view widget for our plugin categories tree */
+	TSharedPtr<STreeView<TSharedPtr<FPluginCategory>>> PluginCategoryTreeView;
+
+	/** Root list of categories */
+	TArray<TSharedPtr<FPluginCategory>> RootPluginCategories;
+};
+
diff --git a/Engine/Plugins/Editor/PluginBrowser/Source/PluginBrowser/Private/SPluginCategoryTreeItem.cpp b/Engine/Plugins/Editor/PluginBrowser/Source/PluginBrowser/Private/SPluginCategoryTreeItem.cpp
deleted file mode 100644
diff --git a/Engine/Plugins/Editor/PluginBrowser/Source/PluginBrowser/Private/SPluginCategoryTreeItem.h b/Engine/Plugins/Editor/PluginBrowser/Source/PluginBrowser/Private/SPluginCategoryTreeItem.h
deleted file mode 100644
diff --git a/Engine/Plugins/Editor/PluginBrowser/Source/PluginBrowser/Private/SPluginList.cpp b/Engine/Plugins/Editor/PluginBrowser/Source/PluginBrowser/Private/SPluginList.cpp
deleted file mode 100644
diff --git a/Engine/Plugins/Editor/PluginBrowser/Source/PluginBrowser/Private/SPluginList.h b/Engine/Plugins/Editor/PluginBrowser/Source/PluginBrowser/Private/SPluginList.h
deleted file mode 100644
diff --git a/Engine/Plugins/Editor/PluginBrowser/Source/PluginBrowser/Private/SPluginListTile.cpp b/Engine/Plugins/Editor/PluginBrowser/Source/PluginBrowser/Private/SPluginListTile.cpp
deleted file mode 100644
diff --git a/Engine/Plugins/Editor/PluginBrowser/Source/PluginBrowser/Private/SPluginListTile.h b/Engine/Plugins/Editor/PluginBrowser/Source/PluginBrowser/Private/SPluginListTile.h
deleted file mode 100644
diff --git a/Engine/Plugins/Editor/PluginBrowser/Source/PluginBrowser/Private/SPluginTile.cpp b/Engine/Plugins/Editor/PluginBrowser/Source/PluginBrowser/Private/SPluginTile.cpp
new file mode 100644
--- /dev/null
+++ b/Engine/Plugins/Editor/PluginBrowser/Source/PluginBrowser/Private/SPluginTile.cpp
@@ -0,0 +1,482 @@
+// Copyright 1998-2015 Epic Games, Inc. All Rights Reserved.
+
+#include "PluginBrowserPrivatePCH.h"
+#include "SPluginTile.h"
+#include "SPluginBrowser.h"
+#include "SPluginTileList.h"
+#include "PluginStyle.h"
+#include "GameProjectGenerationModule.h"
+#include "IDetailsView.h"
+#include "SHyperlink.h"
+#include "PluginMetadataObject.h"
+#include "IProjectManager.h"
+#include "PluginBrowserModule.h"
+#include "ISourceControlModule.h"
+
+#define LOCTEXT_NAMESPACE "PluginListTile"
+
+
+void SPluginTile::Construct( const FArguments& Args, const TSharedRef<SPluginTileList> Owner, TSharedRef<IPlugin> InPlugin )
+{
+	OwnerWeak = Owner;
+	Plugin = InPlugin;
+
+	RecreateWidgets();
+}
+
+void SPluginTile::RecreateWidgets()
+{
+	const float PaddingAmount = FPluginStyle::Get()->GetFloat( "PluginTile.Padding" );
+	const float ThumbnailImageSize = FPluginStyle::Get()->GetFloat( "PluginTile.ThumbnailImageSize" );
+
+	// @todo plugedit: Also display whether plugin is editor-only, runtime-only, developer or a combination?
+	//		-> Maybe a filter for this too?  (show only editor plugins, etc.)
+	// @todo plugedit: Indicate whether plugin has content?  Filter to show only content plugins, and vice-versa?
+
+	// @todo plugedit: Maybe we should do the FileExists check ONCE at plugin load time and not at query time
+
+	const FPluginDescriptor& PluginDescriptor = Plugin->GetDescriptor();
+
+	// Plugin thumbnail image
+	FString Icon128FilePath = Plugin->GetBaseDir() / TEXT("Resources/Icon128.png");
+	if(!FPlatformFileManager::Get().GetPlatformFile().FileExists(*Icon128FilePath))
+	{
+		Icon128FilePath = IPluginManager::Get().FindPlugin(TEXT("PluginBrowser"))->GetBaseDir() / TEXT("Resources/DefaultIcon128.png");
+	}
+
+	const FName BrushName( *Icon128FilePath );
+	const FIntPoint Size = FSlateApplication::Get().GetRenderer()->GenerateDynamicImageResource(BrushName);
+	if ((Size.X > 0) && (Size.Y > 0))
+	{
+		PluginIconDynamicImageBrush = MakeShareable(new FSlateDynamicImageBrush(BrushName, FVector2D(Size.X, Size.Y)));
+	}
+
+	// create documentation link
+	TSharedPtr<SWidget> DocumentationWidget;
+	{
+		if (PluginDescriptor.DocsURL.IsEmpty())
+		{
+			DocumentationWidget = SNullWidget::NullWidget;
+		}
+		else
+		{
+			FString DocsURL = PluginDescriptor.DocsURL;
+			DocumentationWidget = SNew(SHorizontalBox)
+
+			+ SHorizontalBox::Slot()
+				.AutoWidth()
+				.VAlign(VAlign_Center)
+				[
+					SNew(SImage)
+						.ColorAndOpacity(FSlateColor::UseForeground())
+						.Image(FEditorStyle::GetBrush("MessageLog.Docs"))
+				]
+
+			+ SHorizontalBox::Slot()
+				.AutoWidth()
+				.VAlign(VAlign_Center)
+				.Padding(2.0f, 0.0f, 0.0f, 0.0f)
+				[
+					SNew(SHyperlink)
+						.Text(LOCTEXT("DocumentationLink", "Documentation"))
+						.ToolTipText(LOCTEXT("NavigateToCreatedByURL", "Open the plug-in's online documentation"))
+						.OnNavigate_Lambda([=]() { FPlatformProcess::LaunchURL(*DocsURL, nullptr, nullptr); })
+				];
+		}
+	}
+
+	// create vendor link
+	TSharedPtr<SWidget> CreatedByWidget;
+	{
+		if (PluginDescriptor.CreatedBy.IsEmpty())
+		{
+			CreatedByWidget = SNullWidget::NullWidget;
+		}
+		else if (PluginDescriptor.CreatedByURL.IsEmpty())
+		{
+			CreatedByWidget = SNew(SHorizontalBox)
+
+			+ SHorizontalBox::Slot()
+				.AutoWidth()
+				.VAlign(VAlign_Center)
+				[
+					SNew(SImage)
+						.ColorAndOpacity(FSlateColor::UseForeground())
+						.Image(FEditorStyle::GetBrush("ContentBrowser.AssetTreeFolderDeveloper"))
+				]
+
+			+ SHorizontalBox::Slot()
+				.AutoWidth()
+				.VAlign(VAlign_Center)
+				.Padding(2.0f, 0.0f, 0.0f, 0.0f)
+				[
+					SNew(STextBlock)
+						.Text(FText::FromString(PluginDescriptor.CreatedBy))
+				];
+		}
+		else
+		{
+			FString CreatedByURL = PluginDescriptor.CreatedByURL;
+			CreatedByWidget = SNew(SHorizontalBox)
+
+			+ SHorizontalBox::Slot()
+				.AutoWidth()
+				.VAlign(VAlign_Center)
+				[
+					SNew(SImage)
+						.ColorAndOpacity(FSlateColor::UseForeground())
+						.Image(FEditorStyle::GetBrush("MessageLog.Url"))
+				]
+
+			+ SHorizontalBox::Slot()
+				.AutoWidth()
+				.VAlign(VAlign_Center)
+				.Padding(2.0f, 0.0f, 0.0f, 0.0f)
+				[				
+					SNew(SHyperlink)
+						.Text(FText::FromString(PluginDescriptor.CreatedBy))
+						.ToolTipText(FText::Format(LOCTEXT("NavigateToCreatedByURL", "Visit the vendor's web site ({0})"), FText::FromString(CreatedByURL)))
+						.OnNavigate_Lambda([=]() { FPlatformProcess::LaunchURL(*CreatedByURL, nullptr, nullptr); })
+				];
+		}
+	}
+
+	ChildSlot
+	[
+		SNew(SBorder)
+			.BorderImage(FEditorStyle::GetBrush("NoBorder"))
+			.Padding(PaddingAmount)
+			[
+				SNew(SBorder)
+					.BorderImage(FEditorStyle::GetBrush("ToolPanel.GroupBorder"))
+					.Padding(PaddingAmount)
+					[
+						SNew(SHorizontalBox)
+
+						// Thumbnail image
+						+ SHorizontalBox::Slot()
+							.Padding(PaddingAmount)
+							.AutoWidth()
+							[
+								SNew(SBox)
+									.WidthOverride(ThumbnailImageSize)
+									.HeightOverride(ThumbnailImageSize)
+									[
+										SNew(SImage)
+											.Image(PluginIconDynamicImageBrush.IsValid() ? PluginIconDynamicImageBrush.Get() : nullptr)
+									]
+							]
+
+						+SHorizontalBox::Slot()
+							[
+								SNew(SVerticalBox)
+
+								+ SVerticalBox::Slot()
+									.AutoHeight()
+									[
+										SNew(SHorizontalBox)
+
+										// Friendly name
+										+SHorizontalBox::Slot()
+										.Padding(PaddingAmount)
+										[
+											SNew(STextBlock)
+												.Text(FText::FromString(PluginDescriptor.FriendlyName))
+												.HighlightText_Raw(&OwnerWeak.Pin()->GetOwner().GetPluginTextFilter(), &FPluginTextFilter::GetRawFilterText)
+												.TextStyle(FPluginStyle::Get(), "PluginTile.NameText")
+										]
+
+										// Version
+										+ SHorizontalBox::Slot()
+											.HAlign(HAlign_Right)
+											.Padding(PaddingAmount)
+											.AutoWidth()
+											[
+												SNew(SHorizontalBox)
+
+												+ SHorizontalBox::Slot()
+													.AutoWidth()
+													.Padding(0.0f, 0.0f, 0.0f, 1.0f) // Lower padding to align font with version number base
+													[
+														SNew(SHorizontalBox)
+
+														// beta version icon
+														+ SHorizontalBox::Slot()
+															.AutoWidth()
+															.VAlign(VAlign_Bottom)
+															.Padding(0.0f, 0.0f, 3.0f, 2.0f)
+															[
+																SNew(SImage)
+																	.Visibility(PluginDescriptor.bIsBetaVersion ? EVisibility::Visible : EVisibility::Collapsed)
+																	.Image(FPluginStyle::Get()->GetBrush("PluginTile.BetaWarning"))
+															]
+
+														// version label
+														+ SHorizontalBox::Slot()
+															.AutoWidth()
+															.VAlign(VAlign_Bottom)
+															[
+																SNew(STextBlock)
+																	.Text(PluginDescriptor.bIsBetaVersion ? LOCTEXT("PluginBetaVersionLabel", "BETA Version ") : LOCTEXT("PluginVersionLabel", "Version "))
+															]
+													]
+
+												+ SHorizontalBox::Slot()
+													.AutoWidth()
+													.VAlign( VAlign_Bottom )
+													.Padding( 0.0f, 0.0f, 2.0f, 0.0f )	// Extra padding from the right edge
+													[
+														SNew(STextBlock)
+															.Text(FText::FromString(PluginDescriptor.VersionName))
+															.TextStyle(FPluginStyle::Get(), "PluginTile.VersionNumberText")
+													]
+											]
+									]
+			
+								+ SVerticalBox::Slot()
+									[
+										SNew(SVerticalBox)
+				
+										// Description
+										+ SVerticalBox::Slot()
+											.Padding( PaddingAmount )
+											[
+												SNew(STextBlock)
+													.Text(FText::FromString(PluginDescriptor.Description))
+													.AutoWrapText(true)
+											]
+
+										+ SVerticalBox::Slot()
+											.Padding(PaddingAmount)
+											.AutoHeight()
+											[
+												SNew(SHorizontalBox)
+
+												// Enable checkbox
+												+ SHorizontalBox::Slot()
+													.Padding(PaddingAmount)
+													.HAlign(HAlign_Left)
+													[
+														SNew(SCheckBox)
+															.OnCheckStateChanged(this, &SPluginTile::OnEnablePluginCheckboxChanged)
+															.IsChecked(this, &SPluginTile::IsPluginEnabled)
+															.ToolTipText(LOCTEXT("EnableDisableButtonToolTip", "Toggles whether this plugin is enabled for your current project.  You may need to restart the program for this change to take effect."))
+															.Content()
+															[
+																SNew(STextBlock)
+																	.Text(LOCTEXT("EnablePluginCheckbox", "Enabled"))
+															]
+													]
+
+												// edit link
+												+ SHorizontalBox::Slot()
+													.HAlign(HAlign_Center)
+													.AutoWidth()
+													.Padding(2.0f, 0.0f, 0.0f, 0.0f)
+													[
+														SNew(SHorizontalBox)
+
+														+ SHorizontalBox::Slot()
+															.AutoWidth()
+															.Padding(PaddingAmount)
+															[
+																SNew(SHyperlink)
+																	.Visibility(this, &SPluginTile::GetAuthoringButtonsVisibility)	
+																	.OnNavigate(this, &SPluginTile::OnEditPlugin)
+																	.Text(LOCTEXT("EditPlugin", "Edit..."))
+															]
+													]
+
+												// docs link
+												+ SHorizontalBox::Slot()
+													.Padding(PaddingAmount)
+													.HAlign(HAlign_Right)
+													[
+														DocumentationWidget.ToSharedRef()
+													]
+
+												// vendor link
+												+ SHorizontalBox::Slot()
+													.AutoWidth()
+													.Padding(12.0f, PaddingAmount, PaddingAmount, PaddingAmount)
+													.HAlign(HAlign_Right)
+													[
+														CreatedByWidget.ToSharedRef()
+													]
+											]
+									]
+							]
+					]
+			]
+	];
+}
+
+
+ECheckBoxState SPluginTile::IsPluginEnabled() const
+{
+	FPluginBrowserModule& PluginBrowserModule = FPluginBrowserModule::Get();
+	if(PluginBrowserModule.PendingEnablePlugins.Contains(Plugin->GetName()))
+	{
+		return PluginBrowserModule.PendingEnablePlugins[Plugin->GetName()]? ECheckBoxState::Checked : ECheckBoxState::Unchecked;;
+	}
+	else
+	{
+		return Plugin->IsEnabled()? ECheckBoxState::Checked : ECheckBoxState::Unchecked;
+	}
+}
+
+void SPluginTile::OnEnablePluginCheckboxChanged(ECheckBoxState NewCheckedState)
+{
+	const bool bNewEnabledState = (NewCheckedState == ECheckBoxState::Checked);
+
+	const FPluginDescriptor& PluginDescriptor = Plugin->GetDescriptor();
+	if (bNewEnabledState && PluginDescriptor.bIsBetaVersion)
+	{
+		FText WarningMessage = FText::Format(
+			LOCTEXT("Warning_EnablingBetaPlugin", "Plugin '{0}' is a beta version and might be unstable or removed without notice. Please use with caution. Are you sure you want to enable the plugin?"),
+			FText::FromString(PluginDescriptor.FriendlyName));
+
+		if (EAppReturnType::No == FMessageDialog::Open(EAppMsgType::YesNo, WarningMessage))
+		{
+			// user chose to keep beta plug-in disabled
+			return;
+		}
+	}
+
+	FGameProjectGenerationModule::Get().TryMakeProjectFileWriteable(FPaths::GetProjectFilePath());
+	FText FailMessage;
+
+	if (!IProjectManager::Get().SetPluginEnabled(Plugin->GetName(), bNewEnabledState, FailMessage))
+	{
+		FMessageDialog::Open(EAppMsgType::Ok, FailMessage);
+	}
+	else
+	{
+		FPluginBrowserModule& PluginBrowserModule = FPluginBrowserModule::Get();
+		if(Plugin->IsEnabled() == bNewEnabledState)
+		{
+			PluginBrowserModule.PendingEnablePlugins.Remove(Plugin->GetName());
+		}
+		else
+		{
+			PluginBrowserModule.PendingEnablePlugins.FindOrAdd(Plugin->GetName()) = bNewEnabledState;
+		}
+	}
+
+
+}
+
+EVisibility SPluginTile::GetAuthoringButtonsVisibility() const
+{
+	return (FApp::IsEngineInstalled() && Plugin->GetLoadedFrom() == EPluginLoadedFrom::Engine)? EVisibility::Hidden : EVisibility::Visible;
+}
+
+void SPluginTile::OnEditPlugin()
+{
+	// Construct the plugin metadata object using the descriptor for this plugin
+	UPluginMetadataObject* MetadataObject = NewObject<UPluginMetadataObject>();
+	MetadataObject->TargetIconPath = Plugin->GetBaseDir() / TEXT("Resources/Icon128.png");
+	MetadataObject->PopulateFromDescriptor(Plugin->GetDescriptor());
+	MetadataObject->AddToRoot();
+
+	// Create a property view
+	FPropertyEditorModule& EditModule = FModuleManager::Get().GetModuleChecked<FPropertyEditorModule>("PropertyEditor");
+	TSharedRef<IDetailsView> PropertyView = EditModule.CreateDetailView(FDetailsViewArgs(false, false, false, FDetailsViewArgs::ActorsUseNameArea, true));
+	PropertyView->SetObject(MetadataObject, true);
+
+	// Create the window
+	PropertiesWindow = SNew(SWindow)
+		.SupportsMaximize(false)
+		.SupportsMinimize(false)
+		.SizingRule(ESizingRule::FixedSize)
+		.ClientSize(FVector2D(700.0f, 575.0f))
+		.Title(LOCTEXT("PluginMetadata", "Plugin Properties"))
+		.Content()
+		[
+			SNew(SBorder)
+			.Padding(FMargin(8.0f, 8.0f))
+			[
+				SNew(SVerticalBox)
+
+				+SVerticalBox::Slot()
+				.AutoHeight()
+				.Padding(FMargin(5.0f, 10.0f, 5.0f, 5.0f))
+				[
+					SNew(STextBlock)
+					.Font(FPluginStyle::Get()->GetFontStyle(TEXT("PluginMetadataNameFont")))
+					.Text(FText::FromString(Plugin->GetName()))
+				]
+
+				+ SVerticalBox::Slot()
+				.Padding(5)
+				[
+					PropertyView
+				]
+
+				+ SVerticalBox::Slot()
+				.AutoHeight()
+				.Padding(5)
+				.HAlign(HAlign_Right)
+				[
+					SNew(SButton)
+					.ContentPadding(FMargin(20.0f, 2.0f))
+					.Text(LOCTEXT("OkButtonLabel", "Ok"))
+					.OnClicked(this, &SPluginTile::OnEditPluginFinished, MetadataObject)
+				]
+			]
+		];
+
+	FSlateApplication::Get().AddModalWindow(PropertiesWindow.ToSharedRef(), OwnerWeak.Pin());//Args.ParentWidget);
+}
+
+FReply SPluginTile::OnEditPluginFinished(UPluginMetadataObject* MetadataObject)
+{
+	FPluginDescriptor OldDescriptor = Plugin->GetDescriptor();
+
+	// Update the descriptor with the new metadata
+	FPluginDescriptor NewDescriptor = OldDescriptor;
+	MetadataObject->CopyIntoDescriptor(NewDescriptor);
+	MetadataObject->RemoveFromRoot();
+
+	// Close the properties window
+	PropertiesWindow->RequestDestroyWindow();
+
+	// Write both to strings
+	FString OldText = OldDescriptor.ToString();
+	FString NewText = NewDescriptor.ToString();
+	if(OldText.Compare(NewText, ESearchCase::CaseSensitive) != 0)
+	{
+		FString DescriptorFileName = Plugin->GetDescriptorFileName();
+
+		// First attempt to check out the file if SCC is enabled
+		ISourceControlModule& SourceControlModule = ISourceControlModule::Get();
+		if(SourceControlModule.IsEnabled())
+		{
+			ISourceControlProvider& SourceControlProvider = SourceControlModule.GetProvider();
+			TSharedPtr<ISourceControlState, ESPMode::ThreadSafe> SourceControlState = SourceControlProvider.GetState(DescriptorFileName, EStateCacheUsage::ForceUpdate);
+			if(SourceControlState.IsValid() && SourceControlState->CanCheckout())
+			{
+				SourceControlProvider.Execute(ISourceControlOperation::Create<FCheckOut>(), DescriptorFileName);
+			}
+		}
+
+		// Write to the file and update the in-memory metadata
+		FText FailReason;
+		if(!Plugin->UpdateDescriptor(NewDescriptor, FailReason))
+		{
+			FMessageDialog::Open(EAppMsgType::Ok, FailReason);
+		}
+
+		// Recreate the widgets on this tile
+		RecreateWidgets();
+
+		// Refresh the parent too
+		if(OwnerWeak.IsValid())
+		{
+			OwnerWeak.Pin()->GetOwner().SetNeedsRefresh();
+		}
+	}
+	return FReply::Handled();
+}
+
+#undef LOCTEXT_NAMESPACE
diff --git a/Engine/Plugins/Editor/PluginBrowser/Source/PluginBrowser/Private/SPluginTile.h b/Engine/Plugins/Editor/PluginBrowser/Source/PluginBrowser/Private/SPluginTile.h
new file mode 100644
--- /dev/null
+++ b/Engine/Plugins/Editor/PluginBrowser/Source/PluginBrowser/Private/SPluginTile.h
@@ -0,0 +1,54 @@
+// Copyright 1998-2015 Epic Games, Inc. All Rights Reserved.
+
+#pragma once
+
+/**
+ * Widget that represents a "tile" for a single plugin in our plugins list
+ */
+class SPluginTile : public SCompoundWidget
+{
+	SLATE_BEGIN_ARGS( SPluginTile )
+	{
+	}
+
+	SLATE_END_ARGS()
+
+
+	/** Widget constructor */
+	void Construct( const FArguments& Args, const TSharedRef< class SPluginTileList > Owner, TSharedRef<IPlugin> Plugin );
+
+private:
+
+	/** Updates the contents of this tile */
+	void RecreateWidgets();
+
+	/** Returns the checked state for the enabled checkbox */
+	ECheckBoxState IsPluginEnabled() const;
+
+	/** Called when the enabled checkbox is clicked */
+	void OnEnablePluginCheckboxChanged(ECheckBoxState NewCheckedState);
+
+	/** Used to determine whether to show the edit and package buttons for this plugin */
+	EVisibility GetAuthoringButtonsVisibility() const;
+
+	/** Called when the 'edit' hyperlink is clicked */
+	void OnEditPlugin();
+
+	/** Called when the 'edit' hyperlink is clicked */
+	FReply OnEditPluginFinished(UPluginMetadataObject* MetadataObject);
+
+private:
+
+	/** The item we're representing the in tree */
+	TSharedPtr<IPlugin> Plugin;
+
+	/** Weak pointer back to its owner */
+	TWeakPtr< class SPluginTileList > OwnerWeak;
+
+	/** Dialog for editing plugin properties */
+	TSharedPtr< SWindow > PropertiesWindow;
+
+	/** Brush resource for the image that is dynamically loaded */
+	TSharedPtr< FSlateDynamicImageBrush > PluginIconDynamicImageBrush;
+};
+
diff --git a/Engine/Plugins/Editor/PluginBrowser/Source/PluginBrowser/Private/SPluginTileList.cpp b/Engine/Plugins/Editor/PluginBrowser/Source/PluginBrowser/Private/SPluginTileList.cpp
new file mode 100644
--- /dev/null
+++ b/Engine/Plugins/Editor/PluginBrowser/Source/PluginBrowser/Private/SPluginTileList.cpp
@@ -0,0 +1,130 @@
+// Copyright 1998-2015 Epic Games, Inc. All Rights Reserved.
+
+#include "PluginBrowserPrivatePCH.h"
+#include "SPluginTileList.h"
+#include "SPluginBrowser.h"
+#include "SPluginTile.h"
+#include "IPluginManager.h"
+
+
+#define LOCTEXT_NAMESPACE "PluginList"
+
+
+void SPluginTileList::Construct( const FArguments& Args, const TSharedRef< SPluginBrowser > Owner )
+{
+	OwnerWeak = Owner;
+
+	// Find out when the plugin text filter changes
+	Owner->GetPluginTextFilter().OnChanged().AddSP( this, &SPluginTileList::OnPluginTextFilterChanged );
+
+	bIsActiveTimerRegistered = false;
+	RebuildAndFilterPluginList();
+
+	// @todo plugedit: Have optional compact version with only plugin icon + name + version?  Only expand selected?
+
+	PluginListView = 
+		SNew( SListView<TSharedRef<IPlugin>> )
+
+		.SelectionMode( ESelectionMode::None )		// No need to select plugins!
+
+		.ListItemsSource( &PluginListItems )
+		.OnGenerateRow( this, &SPluginTileList::PluginListView_OnGenerateRow )
+		;
+
+	ChildSlot.AttachWidget( PluginListView.ToSharedRef() );
+}
+
+
+SPluginTileList::~SPluginTileList()
+{
+	const TSharedPtr< SPluginBrowser > Owner( OwnerWeak.Pin() );
+	if( Owner.IsValid() )
+	{
+		Owner->GetPluginTextFilter().OnChanged().RemoveAll( this );
+	}
+}
+
+
+/** @return Gets the owner of this list */
+SPluginBrowser& SPluginTileList::GetOwner()
+{
+	return *OwnerWeak.Pin();
+}
+
+
+TSharedRef<ITableRow> SPluginTileList::PluginListView_OnGenerateRow( TSharedRef<IPlugin> Item, const TSharedRef<STableViewBase>& OwnerTable )
+{
+	return
+		SNew( STableRow< TSharedRef<IPlugin> >, OwnerTable )
+		[
+			SNew( SPluginTile, SharedThis( this ), Item )
+		];
+}
+
+
+
+void SPluginTileList::RebuildAndFilterPluginList()
+{
+	// Build up the initial list of plugins
+	{
+		PluginListItems.Reset();
+
+		// Get the currently selected category
+		const TSharedPtr<FPluginCategory>& SelectedCategory = OwnerWeak.Pin()->GetSelectedCategory();
+		if( SelectedCategory.IsValid() )
+		{
+			for(TSharedRef<IPlugin> Plugin: SelectedCategory->Plugins)
+			{
+				if(OwnerWeak.Pin()->GetPluginTextFilter().PassesFilter(&Plugin.Get()))
+				{
+					PluginListItems.Add(Plugin);
+				}
+			}
+		}
+
+		// Sort the plugins alphabetically
+		{
+			struct FPluginListItemSorter
+			{
+				bool operator()( const TSharedRef<IPlugin>& A, const TSharedRef<IPlugin>& B ) const
+				{
+					return A->GetDescriptor().FriendlyName < B->GetDescriptor().FriendlyName;
+				}
+			};
+			PluginListItems.Sort( FPluginListItemSorter() );
+		}
+	}
+
+
+	// Update the list widget
+	if( PluginListView.IsValid() )
+	{
+		PluginListView->RequestListRefresh();
+	}
+}
+
+EActiveTimerReturnType SPluginTileList::TriggerListRebuild(double InCurrentTime, float InDeltaTime)
+{
+	RebuildAndFilterPluginList();
+
+	bIsActiveTimerRegistered = false;
+	return EActiveTimerReturnType::Stop;
+}
+
+void SPluginTileList::OnPluginTextFilterChanged()
+{
+	SetNeedsRefresh();
+}
+
+
+void SPluginTileList::SetNeedsRefresh()
+{
+	if (!bIsActiveTimerRegistered)
+	{
+		bIsActiveTimerRegistered = true;
+		RegisterActiveTimer(0.f, FWidgetActiveTimerDelegate::CreateSP(this, &SPluginTileList::TriggerListRebuild));
+	}
+}
+
+
+#undef LOCTEXT_NAMESPACE
\ No newline at end of file
diff --git a/Engine/Plugins/Editor/PluginBrowser/Source/PluginBrowser/Private/SPluginTileList.h b/Engine/Plugins/Editor/PluginBrowser/Source/PluginBrowser/Private/SPluginTileList.h
new file mode 100644
--- /dev/null
+++ b/Engine/Plugins/Editor/PluginBrowser/Source/PluginBrowser/Private/SPluginTileList.h
@@ -0,0 +1,57 @@
+// Copyright 1998-2015 Epic Games, Inc. All Rights Reserved.
+
+#pragma once
+
+/**
+ * A filtered list of plugins, driven by a category selection
+ */
+class SPluginTileList : public SCompoundWidget
+{
+	SLATE_BEGIN_ARGS( SPluginTileList )
+	{
+	}
+
+	SLATE_END_ARGS()
+
+	/** Widget constructor */
+	void Construct( const FArguments& Args, const TSharedRef< class SPluginBrowser > Owner );
+
+	/** Destructor */
+	~SPluginTileList();
+
+	/** @return Gets the owner of this list */
+	SPluginBrowser& GetOwner();
+
+	/** Called to invalidate the list */
+	void SetNeedsRefresh();
+
+
+private:
+
+	/** Called when the plugin text filter has changed what its filtering */
+	void OnPluginTextFilterChanged();
+
+	/** Called to generate a widget for the specified list item */
+	TSharedRef<ITableRow> PluginListView_OnGenerateRow(TSharedRef<IPlugin> Item, const TSharedRef<STableViewBase>& OwnerTable );
+
+	/** Rebuilds the list of plugins from scratch and applies filtering. */
+	void RebuildAndFilterPluginList();
+
+	/** One-off active timer to trigger a full refresh when something has changed with either our filtering or the loaded plugin set */
+	EActiveTimerReturnType TriggerListRebuild(double InCurrentTime, float InDeltaTime);
+
+private:
+
+	/** Weak pointer back to its owner */
+	TWeakPtr< class SPluginBrowser > OwnerWeak;
+
+	/** The list view widget for our plugins list */
+	TSharedPtr<SListView<TSharedRef<IPlugin>>> PluginListView;
+
+	/** List of everything that we want to display in the plugin list */
+	TArray<TSharedRef<IPlugin>> PluginListItems;
+
+	/** Whether the active timer to refresh the list is registered */
+	bool bIsActiveTimerRegistered;
+};
+
