Add a separate category for "installed" plugins, and hide categories if they're empty.

diff --git a/Engine/Plugins/Editor/PluginBrowser/Source/PluginBrowser/Private/SPluginCategory.cpp b/Engine/Plugins/Editor/PluginBrowser/Source/PluginBrowser/Private/SPluginCategory.cpp
--- a/Engine/Plugins/Editor/PluginBrowser/Source/PluginBrowser/Private/SPluginCategory.cpp
+++ b/Engine/Plugins/Editor/PluginBrowser/Source/PluginBrowser/Private/SPluginCategory.cpp
@@ -78,11 +78,11 @@ const FSlateBrush* SPluginCategory::GetIconBrush() const
 	}
 	else if (Category->Name == TEXT("Installed"))
 	{
-		return FPluginStyle::Get()->GetBrush( "CategoryTreeItem.BuiltIn" );
+		return FPluginStyle::Get()->GetBrush( "CategoryTreeItem.Installed");
 	}
 	else
 	{
-		return FPluginStyle::Get()->GetBrush( "CategoryTreeItem.Installed" );
+		return FPluginStyle::Get()->GetBrush( "CategoryTreeItem.BuiltIn" );
 	}
 }
 
diff --git a/Engine/Plugins/Editor/PluginBrowser/Source/PluginBrowser/Private/SPluginCategoryTree.cpp b/Engine/Plugins/Editor/PluginBrowser/Source/PluginBrowser/Private/SPluginCategoryTree.cpp
--- a/Engine/Plugins/Editor/PluginBrowser/Source/PluginBrowser/Private/SPluginCategoryTree.cpp
+++ b/Engine/Plugins/Editor/PluginBrowser/Source/PluginBrowser/Private/SPluginCategoryTree.cpp
@@ -14,38 +14,35 @@ void SPluginCategoryTree::Construct( const FArguments& Args, const TSharedRef< S
 {
 	OwnerWeak = Owner;
 
-	RootPluginCategories.Reset();
+	// Create the root categories
+	BuiltInCategory = MakeShareable(new FPluginCategory(NULL, TEXT("Built-In"), LOCTEXT("BuiltInCategoryName", "Built-In")));
+	InstalledCategory = MakeShareable(new FPluginCategory(NULL, TEXT("Installed"), LOCTEXT("InstalledCategoryName", "Installed")));
+	ProjectCategory = MakeShareable(new FPluginCategory(NULL, TEXT("Project"), LOCTEXT("ProjectCategoryName", "Project")));
 
-	// Add the root level categories
-	TSharedRef<FPluginCategory> InstalledCategory = MakeShareable(new FPluginCategory(NULL, TEXT("Installed"), LOCTEXT("InstalledCategoryName", "Installed")));
-	RootPluginCategories.Add( InstalledCategory );
-	TSharedRef<FPluginCategory> ProjectCategory = MakeShareable(new FPluginCategory(NULL, TEXT("Project"), LOCTEXT("ProjectCategoryName", "Project")));
-	RootPluginCategories.Add( ProjectCategory );
-
-	PluginCategoryTreeView =
+	// Create the tree view control
+	TreeView =
 		SNew( STreeView<TSharedPtr<FPluginCategory>> )
 
 		// For now we only support selecting a single folder in the tree
 		.SelectionMode( ESelectionMode::Single )
 		.ClearSelectionOnClick( false )		// Don't allow user to select nothing.  We always expect a category to be selected!
 
-		.TreeItemsSource( &RootPluginCategories )
+		.TreeItemsSource( &RootCategories )
 		.OnGenerateRow( this, &SPluginCategoryTree::PluginCategoryTreeView_OnGenerateRow ) 
 		.OnGetChildren( this, &SPluginCategoryTree::PluginCategoryTreeView_OnGetChildren )
 
 		.OnSelectionChanged( this, &SPluginCategoryTree::PluginCategoryTreeView_OnSelectionChanged )
 		;
 
-	// Expand the root categories by default
-	for( auto RootCategoryIt( RootPluginCategories.CreateConstIterator() ); RootCategoryIt; ++RootCategoryIt )
+	RebuildAndFilterCategoryTree();
+
+	// Expand all the root categories by default
+	for(TSharedPtr<FPluginCategory> RootCategory: RootCategories)
 	{
-		const auto& Category = *RootCategoryIt;
-		PluginCategoryTreeView->SetItemExpansion( Category, true );
+		TreeView->SetItemExpansion( RootCategory, true );
 	}
 
-	RebuildAndFilterCategoryTree();
-
-	ChildSlot.AttachWidget( PluginCategoryTreeView.ToSharedRef() );
+	ChildSlot.AttachWidget( TreeView.ToSharedRef() );
 }
 
 
@@ -65,7 +62,7 @@ void SPluginCategoryTree::RebuildAndFilterCategoryTree()
 {
 	// Get a plugin from the currently selected category, so we can track it if it's removed
 	TSharedPtr<IPlugin> TrackPlugin = nullptr;
-	for(TSharedPtr<FPluginCategory> SelectedItem: PluginCategoryTreeView->GetSelectedItems())
+	for(TSharedPtr<FPluginCategory> SelectedItem: TreeView->GetSelectedItems())
 	{
 		if(SelectedItem->Plugins.Num() > 0)
 		{
@@ -75,7 +72,7 @@ void SPluginCategoryTree::RebuildAndFilterCategoryTree()
 	}
 
 	// Clear the list of plugins in each current category
-	for(TSharedPtr<FPluginCategory> RootCategory: RootPluginCategories)
+	for(TSharedPtr<FPluginCategory> RootCategory: RootCategories)
 	{
 		for(TSharedPtr<FPluginCategory> Category: RootCategory->SubCategories)
 		{
@@ -89,13 +86,17 @@ void SPluginCategoryTree::RebuildAndFilterCategoryTree()
 	{
 		// Figure out which base category this plugin belongs in
 		TSharedPtr<FPluginCategory> RootCategory;
-		if( Plugin->GetLoadedFrom() == EPluginLoadedFrom::Engine )
+		if(Plugin->GetDescriptor().bInstalled)
 		{
-			RootCategory = RootPluginCategories[0];
+			RootCategory = InstalledCategory;
+		}
+		else if(Plugin->GetLoadedFrom() == EPluginLoadedFrom::Engine)
+		{
+			RootCategory = BuiltInCategory;
 		}
 		else
 		{
-			RootCategory = RootPluginCategories[1];
+			RootCategory = ProjectCategory;
 		}
 
 		// Get the subcategory for this plugin
@@ -134,7 +135,7 @@ void SPluginCategoryTree::RebuildAndFilterCategoryTree()
 	}
 
 	// Remove any empty categories, keeping track of which items are still selected
-	for(TSharedPtr<FPluginCategory> RootCategory: RootPluginCategories)
+	for(TSharedPtr<FPluginCategory> RootCategory: RootCategories)
 	{
 		for(int32 Idx = 0; Idx < RootCategory->SubCategories.Num(); Idx++)
 		{
@@ -145,27 +146,41 @@ void SPluginCategoryTree::RebuildAndFilterCategoryTree()
 		}
 	}
 
+	// Build the new list of root plugin categories
+	RootCategories.Reset();
+	if(InstalledCategory->SubCategories.Num() > 0 || InstalledCategory->Plugins.Num() > 0)
+	{
+		RootCategories.Add(InstalledCategory);
+	}
+	if(BuiltInCategory->SubCategories.Num() > 0 || BuiltInCategory->Plugins.Num() > 0)
+	{
+		RootCategories.Add(BuiltInCategory);
+	}
+	if(ProjectCategory->SubCategories.Num() > 0 || ProjectCategory->Plugins.Num() > 0)
+	{
+		RootCategories.Add(ProjectCategory);
+	}
+
 	// Sort every single category alphabetically
-	for(TSharedPtr<FPluginCategory> RootCategory: RootPluginCategories)
+	for(TSharedPtr<FPluginCategory> RootCategory: RootCategories)
 	{
 		RootCategory->SubCategories.Sort([](const TSharedPtr<FPluginCategory>& A, const TSharedPtr<FPluginCategory>& B) -> bool { return A->DisplayName.CompareTo(B->DisplayName) < 0; });
 	}
 
 	// Refresh the view
-	PluginCategoryTreeView->RequestTreeRefresh();
+	TreeView->RequestTreeRefresh();
 
 	// Make sure we have something selected
-	if(SelectCategory.IsValid())
+	if(RootCategories.Num() > 0)
 	{
-		PluginCategoryTreeView->SetSelection(SelectCategory);
-	}
-	else if(RootPluginCategories[0]->SubCategories.Num() > 0)
-	{
-		PluginCategoryTreeView->SetSelection(RootPluginCategories[0]->SubCategories[0]);
-	}
-	else
-	{
-		PluginCategoryTreeView->SetSelection(RootPluginCategories[0]);
+		if(SelectCategory.IsValid())
+		{
+			TreeView->SetSelection(SelectCategory);
+		}
+		else if(RootCategories.Num() > 0 && RootCategories[0]->SubCategories.Num() > 0)
+		{
+			TreeView->SetSelection(RootCategories[0]->SubCategories[0]);
+		}
 	}
 }
 
@@ -194,9 +209,9 @@ void SPluginCategoryTree::PluginCategoryTreeView_OnSelectionChanged(TSharedPtr<F
 
 TSharedPtr<FPluginCategory> SPluginCategoryTree::GetSelectedCategory() const
 {
-	if( PluginCategoryTreeView.IsValid() )
+	if( TreeView.IsValid() )
 	{
-		auto SelectedItems = PluginCategoryTreeView->GetSelectedItems();
+		auto SelectedItems = TreeView->GetSelectedItems();
 		if( SelectedItems.Num() > 0 )
 		{
 			const auto& SelectedCategoryItem = SelectedItems[ 0 ];
@@ -209,15 +224,15 @@ TSharedPtr<FPluginCategory> SPluginCategoryTree::GetSelectedCategory() const
 
 void SPluginCategoryTree::SelectCategory( const TSharedPtr<FPluginCategory>& CategoryToSelect )
 {
-	if( ensure( PluginCategoryTreeView.IsValid() ) )
+	if( ensure( TreeView.IsValid() ) )
 	{
-		PluginCategoryTreeView->SetSelection( CategoryToSelect );
+		TreeView->SetSelection( CategoryToSelect );
 	}
 }
 
 bool SPluginCategoryTree::IsItemExpanded( const TSharedPtr<FPluginCategory> Item ) const
 {
-	return PluginCategoryTreeView->IsItemExpanded( Item );
+	return TreeView->IsItemExpanded( Item );
 }
 
 void SPluginCategoryTree::SetNeedsRefresh()
diff --git a/Engine/Plugins/Editor/PluginBrowser/Source/PluginBrowser/Private/SPluginCategoryTree.h b/Engine/Plugins/Editor/PluginBrowser/Source/PluginBrowser/Private/SPluginCategoryTree.h
--- a/Engine/Plugins/Editor/PluginBrowser/Source/PluginBrowser/Private/SPluginCategoryTree.h
+++ b/Engine/Plugins/Editor/PluginBrowser/Source/PluginBrowser/Private/SPluginCategoryTree.h
@@ -63,9 +63,18 @@ private:
 	TWeakPtr< class SPluginBrowser > OwnerWeak;
 
 	/** The tree view widget for our plugin categories tree */
-	TSharedPtr<STreeView<TSharedPtr<FPluginCategory>>> PluginCategoryTreeView;
+	TSharedPtr<STreeView<TSharedPtr<FPluginCategory>>> TreeView;
 
 	/** Root list of categories */
-	TArray<TSharedPtr<FPluginCategory>> RootPluginCategories;
+	TArray<TSharedPtr<FPluginCategory>> RootCategories;
+
+	/** Category for built-in plugins */
+	TSharedPtr<FPluginCategory> BuiltInCategory;
+
+	/** Category for installed plugins */
+	TSharedPtr<FPluginCategory> InstalledCategory;
+
+	/** Category for project plugins */
+	TSharedPtr<FPluginCategory> ProjectCategory;
 };
 
diff --git a/Engine/Source/Programs/UnrealBuildTool/System/PluginDescriptor.cs b/Engine/Source/Programs/UnrealBuildTool/System/PluginDescriptor.cs
--- a/Engine/Source/Programs/UnrealBuildTool/System/PluginDescriptor.cs
+++ b/Engine/Source/Programs/UnrealBuildTool/System/PluginDescriptor.cs
@@ -63,6 +63,9 @@ namespace UnrealBuildTool
 		// Marks the plugin as beta in the UI
 		public bool bIsBetaVersion;
 
+		// Set for plugins which are installed
+		public bool bInstalled;
+
 		/// <summary>
 		/// Private constructor. This object should not be created directly; read it from disk using FromFile() instead.
 		/// </summary>
@@ -130,6 +133,7 @@ namespace UnrealBuildTool
 				RawObject.TryGetBoolField("EnabledByDefault", out Descriptor.bEnabledByDefault);
 				RawObject.TryGetBoolField("CanContainContent", out Descriptor.bCanContainContent);
 				RawObject.TryGetBoolField("IsBetaVersion", out Descriptor.bIsBetaVersion);
+				RawObject.TryGetBoolField("Installed", out Descriptor.bInstalled);
 
 				return Descriptor;
 			}
diff --git a/Engine/Source/Runtime/Projects/Private/PluginDescriptor.cpp b/Engine/Source/Runtime/Projects/Private/PluginDescriptor.cpp
--- a/Engine/Source/Runtime/Projects/Private/PluginDescriptor.cpp
+++ b/Engine/Source/Runtime/Projects/Private/PluginDescriptor.cpp
@@ -12,7 +12,9 @@ FPluginDescriptor::FPluginDescriptor()
 	, bEnabledByDefault(false)
 	, bCanContainContent(false)
 	, bIsBetaVersion(false)
-{ }
+	, bInstalled(false)
+{ 
+}
 
 
 bool FPluginDescriptor::Load( const FString& FileName, FText& OutFailReason )
@@ -96,6 +98,7 @@ bool FPluginDescriptor::Read(const FString& Text, FText& OutFailReason)
 	Object.TryGetBoolField(TEXT("EnabledByDefault"), bEnabledByDefault);
 	Object.TryGetBoolField(TEXT("CanContainContent"), bCanContainContent);
 	Object.TryGetBoolField(TEXT("IsBetaVersion"), bIsBetaVersion);
+	Object.TryGetBoolField(TEXT("Installed"), bInstalled);
 
 	return true;
 }
@@ -138,6 +141,7 @@ FString FPluginDescriptor::ToString() const
 	Writer.WriteValue(TEXT("EnabledByDefault"), bEnabledByDefault);
 	Writer.WriteValue(TEXT("CanContainContent"), bCanContainContent);
 	Writer.WriteValue(TEXT("IsBetaVersion"), bIsBetaVersion);
+	Writer.WriteValue(TEXT("Installed"), bInstalled);
 
 	Writer.WriteObjectEnd();
 	Writer.Close();
diff --git a/Engine/Source/Runtime/Projects/Public/PluginDescriptor.h b/Engine/Source/Runtime/Projects/Public/PluginDescriptor.h
--- a/Engine/Source/Runtime/Projects/Public/PluginDescriptor.h
+++ b/Engine/Source/Runtime/Projects/Public/PluginDescriptor.h
@@ -70,6 +70,9 @@ struct PROJECTS_API FPluginDescriptor
 	/** Marks the plugin as beta in the UI */
 	bool bIsBetaVersion;
 
+	/** Signifies that the plugin was installed on top of the engine */
+	bool bInstalled;
+
 	/** Constructor. */
 	FPluginDescriptor();
 
