When building the plugin category tree, plugins are now also added to their parent category until we reach the root.

diff --git a/Engine/Plugins/Editor/PluginBrowser/Source/PluginBrowser/Private/SPluginCategoryTree.cpp b/Engine/Plugins/Editor/PluginBrowser/Source/PluginBrowser/Private/SPluginCategoryTree.cpp
--- a/Engine/Plugins/Editor/PluginBrowser/Source/PluginBrowser/Private/SPluginCategoryTree.cpp
+++ b/Engine/Plugins/Editor/PluginBrowser/Source/PluginBrowser/Private/SPluginCategoryTree.cpp
@@ -122,6 +122,13 @@ void SPluginCategoryTree::RebuildAndFilterCategoryTree()
 		// Associate the plugin with the category
 		FoundCategory->Plugins.Add(Plugin);
 
+		TSharedPtr<FPluginCategory> ParentCategory = FoundCategory->ParentCategory.Pin();
+		while (ParentCategory.IsValid())
+		{
+			ParentCategory->Plugins.Add(Plugin);
+			ParentCategory = ParentCategory->ParentCategory.Pin();
+		}
+
 		// Update the selection if this is the plugin we're tracking
 		if(TrackPlugin == Plugin)
 		{
