Fixed some bad font references

UE-8718 - Warning about loading font darta from Roboto-Regular.ttf & Roboto-Bold.ttf

diff --git a/Engine/Plugins/Editor/PluginBrowser/Source/PluginBrowser/Private/PluginStyle.cpp b/Engine/Plugins/Editor/PluginBrowser/Source/PluginBrowser/Private/PluginStyle.cpp
--- a/Engine/Plugins/Editor/PluginBrowser/Source/PluginBrowser/Private/PluginStyle.cpp
+++ b/Engine/Plugins/Editor/PluginBrowser/Source/PluginBrowser/Private/PluginStyle.cpp
@@ -9,7 +9,9 @@
 #define BOX_BRUSH( RelativePath, ... ) FSlateBoxBrush( FPluginStyle::InContent( RelativePath, ".png" ), __VA_ARGS__ )
 #define BORDER_BRUSH( RelativePath, ... ) FSlateBorderBrush( FPluginStyle::InContent( RelativePath, ".png" ), __VA_ARGS__ )
 #define TTF_FONT( RelativePath, ... ) FSlateFontInfo( FPluginStyle::InContent( RelativePath, ".ttf" ), __VA_ARGS__ )
+#define TTF_CORE_FONT(RelativePath, ...) FSlateFontInfo( StyleSet->RootToCoreContentDir( RelativePath, TEXT(".ttf") ), __VA_ARGS__ )
 #define OTF_FONT( RelativePath, ... ) FSlateFontInfo( FPluginStyle::InContent( RelativePath, ".otf" ), __VA_ARGS__ )
+#define OTF_CORE_FONT(RelativePath, ...) FSlateFontInfo( StyleSet->RootToCoreContentDir( RelativePath, TEXT(".otf") ), __VA_ARGS__ )
 
 FString FPluginStyle::InContent( const FString& RelativePath, const ANSICHAR* Extension )
 {
@@ -48,6 +50,7 @@ void FPluginStyle::Initialize()
 	}
 
 	StyleSet = MakeShareable( new FSlateStyleSet("PluginStyle") );
+	StyleSet->SetCoreContentRoot(FPaths::EngineContentDir() / TEXT("Slate"));
 
 	// Plugins Manager
 	{
@@ -167,7 +170,7 @@ void FPluginStyle::Initialize()
 
 		// Metadata editor
 		{
-			StyleSet->Set("PluginMetadataNameFont", TTF_FONT("Fonts/Roboto-Bold", 18));
+			StyleSet->Set("PluginMetadataNameFont", TTF_CORE_FONT("Fonts/Roboto-Bold", 18));
 		}
 	}
 
@@ -178,7 +181,9 @@ void FPluginStyle::Initialize()
 #undef BOX_BRUSH
 #undef BORDER_BRUSH
 #undef TTF_FONT
+#undef TTF_CORE_FONT
 #undef OTF_FONT
+#undef OTF_CORE_FONT
 
 void FPluginStyle::Shutdown()
 {
