Added support URL for plugins

Removed Custom file writing code from Plugin Creator plugin as it hadn't been kept up to date with normal methods and did not write all properties

Also noticed two strings in plugin tile used the same key, allowed documentation and support to preview URL as well

UE-16203

diff --git a/Engine/Plugins/Editor/PluginBrowser/Source/PluginBrowser/Private/PluginMetadataObject.cpp b/Engine/Plugins/Editor/PluginBrowser/Source/PluginBrowser/Private/PluginMetadataObject.cpp
--- a/Engine/Plugins/Editor/PluginBrowser/Source/PluginBrowser/Private/PluginMetadataObject.cpp
+++ b/Engine/Plugins/Editor/PluginBrowser/Source/PluginBrowser/Private/PluginMetadataObject.cpp
@@ -22,6 +22,7 @@ void UPluginMetadataObject::PopulateFromDescriptor(const FPluginDescriptor& InDe
 	CreatedBy = InDescriptor.CreatedBy;
 	CreatedByURL = InDescriptor.CreatedByURL;
 	DocsURL = InDescriptor.DocsURL;
+	SupportURL = InDescriptor.SupportURL;
 	bCanContainContent = InDescriptor.bCanContainContent;
 	bIsBetaVersion = InDescriptor.bIsBetaVersion;
 }
@@ -36,6 +37,7 @@ void UPluginMetadataObject::CopyIntoDescriptor(FPluginDescriptor& OutDescriptor)
 	OutDescriptor.CreatedBy = CreatedBy;
 	OutDescriptor.CreatedByURL = CreatedByURL;
 	OutDescriptor.DocsURL = DocsURL;
+	OutDescriptor.SupportURL = SupportURL;
 	OutDescriptor.bCanContainContent = bCanContainContent;
 	OutDescriptor.bIsBetaVersion = bIsBetaVersion;
 }
diff --git a/Engine/Plugins/Editor/PluginBrowser/Source/PluginBrowser/Private/PluginMetadataObject.h b/Engine/Plugins/Editor/PluginBrowser/Source/PluginBrowser/Private/PluginMetadataObject.h
--- a/Engine/Plugins/Editor/PluginBrowser/Source/PluginBrowser/Private/PluginMetadataObject.h
+++ b/Engine/Plugins/Editor/PluginBrowser/Source/PluginBrowser/Private/PluginMetadataObject.h
@@ -56,6 +56,10 @@ public:
 	UPROPERTY(EditAnywhere, Category = "Details")
 	FString DocsURL;
 
+	/** Support URL/email for this plugin. Email addresses must be prefixed with 'mailto:' */
+	UPROPERTY(EditAnywhere, Category = "Details")
+	FString SupportURL;
+
 	/** Can this plugin contain content? */
 	UPROPERTY(EditAnywhere, Category = "Details")
 	bool bCanContainContent;
diff --git a/Engine/Plugins/Editor/PluginBrowser/Source/PluginBrowser/Private/SPluginTile.cpp b/Engine/Plugins/Editor/PluginBrowser/Source/PluginBrowser/Private/SPluginTile.cpp
--- a/Engine/Plugins/Editor/PluginBrowser/Source/PluginBrowser/Private/SPluginTile.cpp
+++ b/Engine/Plugins/Editor/PluginBrowser/Source/PluginBrowser/Private/SPluginTile.cpp
@@ -52,6 +52,40 @@ void SPluginTile::RecreateWidgets()
 		PluginIconDynamicImageBrush = MakeShareable(new FSlateDynamicImageBrush(BrushName, FVector2D(Size.X, Size.Y)));
 	}
 
+	// create support link
+	TSharedPtr<SWidget> SupportWidget;
+	{
+		if (PluginDescriptor.SupportURL.IsEmpty())
+		{
+			SupportWidget = SNullWidget::NullWidget;
+		}
+		else
+		{
+			FString SupportURL = PluginDescriptor.SupportURL;
+			SupportWidget = SNew(SHorizontalBox)
+
+			+ SHorizontalBox::Slot()
+				.AutoWidth()
+				.VAlign(VAlign_Center)
+				[
+					SNew(SImage)
+					.ColorAndOpacity(FSlateColor::UseForeground())
+					.Image(FEditorStyle::GetBrush("Icons.Contact"))
+				]
+
+			+ SHorizontalBox::Slot()
+				.AutoWidth()
+				.VAlign(VAlign_Center)
+				.Padding(2.0f, 0.0f, 0.0f, 0.0f)
+				[
+					SNew(SHyperlink)
+					.Text(LOCTEXT("SupportLink", "Support"))
+					.ToolTipText(FText::Format(LOCTEXT("NavigateToSupportURL", "Open the plug-in's online support ({0})"), FText::FromString(SupportURL)))
+					.OnNavigate_Lambda([=]() { FPlatformProcess::LaunchURL(*SupportURL, nullptr, nullptr); })
+				];
+		}
+	}
+
 	// create documentation link
 	TSharedPtr<SWidget> DocumentationWidget;
 	{
@@ -80,7 +114,7 @@ void SPluginTile::RecreateWidgets()
 				[
 					SNew(SHyperlink)
 						.Text(LOCTEXT("DocumentationLink", "Documentation"))
-						.ToolTipText(LOCTEXT("NavigateToCreatedByURL", "Open the plug-in's online documentation"))
+						.ToolTipText(FText::Format(LOCTEXT("NavigateToDocumentation", "Open the plug-in's online documentation ({0})"), FText::FromString(DocsURL)))
 						.OnNavigate_Lambda([=]() { FPlatformProcess::LaunchURL(*DocsURL, nullptr, nullptr); })
 				];
 		}
@@ -288,9 +322,18 @@ void SPluginTile::RecreateWidgets()
 															]
 													]
 
+												// support link
+												+SHorizontalBox::Slot()
+													.Padding(PaddingAmount)
+													.HAlign(HAlign_Right)
+													[
+														SupportWidget.ToSharedRef()
+													]
+
 												// docs link
 												+ SHorizontalBox::Slot()
-													.Padding(PaddingAmount)
+													.AutoWidth()
+													.Padding(12.0f, PaddingAmount, PaddingAmount, PaddingAmount)
 													.HAlign(HAlign_Right)
 													[
 														DocumentationWidget.ToSharedRef()
diff --git a/Engine/Plugins/Editor/PluginCreator/Source/PluginCreator/Private/Helpers/PluginDescriptorObject.cpp b/Engine/Plugins/Editor/PluginCreator/Source/PluginCreator/Private/Helpers/PluginDescriptorObject.cpp
--- a/Engine/Plugins/Editor/PluginCreator/Source/PluginCreator/Private/Helpers/PluginDescriptorObject.cpp
+++ b/Engine/Plugins/Editor/PluginCreator/Source/PluginCreator/Private/Helpers/PluginDescriptorObject.cpp
@@ -23,6 +23,7 @@ void UPluginDescriptorObject::FillDescriptor(FPluginDescriptor& OutDescriptor)
 	OutDescriptor.CreatedByURL = CreatedByURL;
 	OutDescriptor.Description = Description;
 	OutDescriptor.DocsURL = DocsURL;
+	OutDescriptor.SupportURL = SupportURL;
 	OutDescriptor.FileVersion = FileVersion;
 	OutDescriptor.FriendlyName = FriendlyName;
 	OutDescriptor.Version = Version;
diff --git a/Engine/Plugins/Editor/PluginCreator/Source/PluginCreator/Private/Helpers/PluginDescriptorObject.h b/Engine/Plugins/Editor/PluginCreator/Source/PluginCreator/Private/Helpers/PluginDescriptorObject.h
--- a/Engine/Plugins/Editor/PluginCreator/Source/PluginCreator/Private/Helpers/PluginDescriptorObject.h
+++ b/Engine/Plugins/Editor/PluginCreator/Source/PluginCreator/Private/Helpers/PluginDescriptorObject.h
@@ -64,6 +64,10 @@ public:
 	UPROPERTY(EditAnywhere, Category = "PluginDescription")
 	FString DocsURL;
 
+	/** Support URL/email for this plugin. Email addresses must be prefixed with 'mailto:' */
+	UPROPERTY(EditAnywhere, Category = "PluginDescription")
+	FString SupportURL;
+
 	/** List of all modules associated with this plugin */
 	TArray<FModuleDescriptor> Modules;
 
diff --git a/Engine/Plugins/Editor/PluginCreator/Source/PluginCreator/Private/Helpers/PluginHelpers.cpp b/Engine/Plugins/Editor/PluginCreator/Source/PluginCreator/Private/Helpers/PluginHelpers.cpp
--- a/Engine/Plugins/Editor/PluginCreator/Source/PluginCreator/Private/Helpers/PluginHelpers.cpp
+++ b/Engine/Plugins/Editor/PluginCreator/Source/PluginCreator/Private/Helpers/PluginHelpers.cpp
@@ -105,38 +105,6 @@ bool FPluginHelpers::CreatePrivatePCHFile(const FString& FolderPath, const FStri
 	return GameProjectUtils::WriteOutputFile(FolderPath / PluginName + TEXT("PrivatePCH.h"), FinalOutput, OutFailReason);
 }
 
-
-bool FPluginHelpers::SavePluginDescriptor(const FString& NewProjectFilename, const FPluginDescriptor& PluginDescriptor)
-{
-	FString Text;
-	TSharedRef< TJsonWriter<> > Writer = TJsonWriterFactory<>::Create(&Text);
-
-	Writer->WriteObjectStart();
-
-	// Write all the simple fields
-
-	Writer->WriteValue(TEXT("FileVersion"), PluginDescriptor.FileVersion);
-
-	Writer->WriteValue(TEXT("FriendlyName"), PluginDescriptor.FriendlyName);
-	Writer->WriteValue(TEXT("Version"), PluginDescriptor.Version);
-	Writer->WriteValue(TEXT("VersionName"), PluginDescriptor.VersionName);
-	Writer->WriteValue(TEXT("CreatedBy"), PluginDescriptor.CreatedBy);
-	Writer->WriteValue(TEXT("CreatedByURL"), PluginDescriptor.CreatedByURL);
-	Writer->WriteValue(TEXT("Category"), PluginDescriptor.Category);
-	Writer->WriteValue(TEXT("Description"), PluginDescriptor.Description);
-	Writer->WriteValue(TEXT("EnabledByDefault"), PluginDescriptor.bEnabledByDefault);
-
-
-	// Write the module list
-	FModuleDescriptor::WriteArray(Writer.Get(), TEXT("Modules"), PluginDescriptor.Modules);
-
-	Writer->WriteObjectEnd();
-
-	Writer->Close();
-
-	return FFileHelper::SaveStringToFile(Text, *NewProjectFilename);
-}
-
 bool FPluginHelpers::CreatePluginStyleFiles(const FString& PrivateFolderPath, const FString& PublicFolderPath, const FString& PluginName, FText& OutFailReason, FString TemplateType)
 {
 	{
diff --git a/Engine/Plugins/Editor/PluginCreator/Source/PluginCreator/Private/Helpers/PluginHelpers.h b/Engine/Plugins/Editor/PluginCreator/Source/PluginCreator/Private/Helpers/PluginHelpers.h
--- a/Engine/Plugins/Editor/PluginCreator/Source/PluginCreator/Private/Helpers/PluginHelpers.h
+++ b/Engine/Plugins/Editor/PluginCreator/Source/PluginCreator/Private/Helpers/PluginHelpers.h
@@ -62,11 +62,4 @@ public:
 	*	@param TemplateType - empty by default; can be Blank, Basic or Advanced
 	*	@return true if file was created, false otherwise*/
 	static bool CreatePrivatePCHFile(const FString& FolderPath, const FString& PluginName, FText& OutFailReason, FString TemplateType = FString(""));
-
-	/** Save plugin descriptor under given path as .uplugin file using JSon
-	*	@param NewProjectFilename - filename with full path
-	*	@param PluginDescriptor - descriptor to save
-	*	@return true if file was created, false otherwise*/
-	static bool SavePluginDescriptor(const FString& NewProjectFilename, const FPluginDescriptor& PluginDescriptor);
-
 };
\ No newline at end of file
diff --git a/Engine/Plugins/Editor/PluginCreator/Source/PluginCreator/Private/UI/SPluginCreatorTabContent.cpp b/Engine/Plugins/Editor/PluginCreator/Source/PluginCreator/Private/UI/SPluginCreatorTabContent.cpp
--- a/Engine/Plugins/Editor/PluginCreator/Source/PluginCreator/Private/UI/SPluginCreatorTabContent.cpp
+++ b/Engine/Plugins/Editor/PluginCreator/Source/PluginCreator/Private/UI/SPluginCreatorTabContent.cpp
@@ -551,7 +551,8 @@ bool SPluginCreatorTabContent::WritePluginDescriptor(const FString& PluginModule
 	}
 
 	// Save the descriptor using JSon
-	if (FPluginHelpers::SavePluginDescriptor(UPluginFilePath, Descriptor))
+	FText FailReason;
+	if (Descriptor.Save(UPluginFilePath, FailReason))
 	{
 		return true;
 	}
diff --git a/Engine/Source/Editor/EditorStyle/Private/SlateEditorStyle.cpp b/Engine/Source/Editor/EditorStyle/Private/SlateEditorStyle.cpp
--- a/Engine/Source/Editor/EditorStyle/Private/SlateEditorStyle.cpp
+++ b/Engine/Source/Editor/EditorStyle/Private/SlateEditorStyle.cpp
@@ -332,6 +332,7 @@ void FSlateEditorStyle::FStyle::SetupGeneralStyles()
 		Set( "Icons.Warning", new IMAGE_BRUSH( "Icons/icon_warning_16x", Icon16x16) );
 		Set( "Icons.Download", new IMAGE_BRUSH( "Icons/icon_Downloads_16x", Icon16x16) );
 		Set( "Icons.Refresh", new IMAGE_BRUSH( "Icons/icon_Refresh_16x", Icon16x16 ) );
+		Set( "Icons.Contact", new IMAGE_BRUSH( "Icons/icon_mail_16x", Icon16x16 ) );
 	}
 
 	Set( "WarningStripe", new IMAGE_BRUSH( "Common/WarningStripe", FVector2D(20,6), FLinearColor::White, ESlateBrushTileType::Horizontal ) );
diff --git a/Engine/Source/Programs/UnrealBuildTool/System/PluginDescriptor.cs b/Engine/Source/Programs/UnrealBuildTool/System/PluginDescriptor.cs
--- a/Engine/Source/Programs/UnrealBuildTool/System/PluginDescriptor.cs
+++ b/Engine/Source/Programs/UnrealBuildTool/System/PluginDescriptor.cs
@@ -55,6 +55,9 @@ namespace UnrealBuildTool
 		// Marketplace URL for this plugin. This URL will be embedded into projects that enable this plugin, so we can redirect to the marketplace if a user doesn't have it installed.
 		public string MarketplaceURL;
 
+		// Support URL/email for this plugin.
+		public string SupportURL;
+
 		// List of all modules associated with this plugin
 		public ModuleDescriptor[] Modules;
 
@@ -128,6 +131,7 @@ namespace UnrealBuildTool
 				RawObject.TryGetStringField("CreatedByURL", out Descriptor.CreatedByURL);
 				RawObject.TryGetStringField("DocsURL", out Descriptor.DocsURL);
 				RawObject.TryGetStringField("MarketplaceURL", out Descriptor.MarketplaceURL);
+				RawObject.TryGetStringField("SupportURL", out Descriptor.SupportURL);
 
 				JsonObject[] ModulesArray;
 				if(RawObject.TryGetObjectArrayField("Modules", out ModulesArray))
@@ -168,6 +172,7 @@ namespace UnrealBuildTool
 				Writer.WriteValue("CreatedByURL", CreatedByURL);
 				Writer.WriteValue("DocsURL", DocsURL);
 				Writer.WriteValue("MarketplaceURL", MarketplaceURL);
+				Writer.WriteValue("SupportURL", SupportURL);
 
 				ModuleDescriptor.WriteArray(Writer, "Modules", Modules);
 
diff --git a/Engine/Source/Runtime/Projects/Private/PluginDescriptor.cpp b/Engine/Source/Runtime/Projects/Private/PluginDescriptor.cpp
--- a/Engine/Source/Runtime/Projects/Private/PluginDescriptor.cpp
+++ b/Engine/Source/Runtime/Projects/Private/PluginDescriptor.cpp
@@ -90,6 +90,7 @@ bool FPluginDescriptor::Read(const FString& Text, FText& OutFailReason)
 	Object.TryGetStringField(TEXT("CreatedByURL"), CreatedByURL);
 	Object.TryGetStringField(TEXT("DocsURL"), DocsURL);
 	Object.TryGetStringField(TEXT("MarketplaceURL"), MarketplaceURL);
+	Object.TryGetStringField(TEXT("SupportURL"), SupportURL);
 
 	if (!FModuleDescriptor::ReadArray(Object, TEXT("Modules"), Modules, OutFailReason))
 	{
@@ -137,6 +138,7 @@ FString FPluginDescriptor::ToString() const
 	Writer.WriteValue(TEXT("CreatedByURL"), CreatedByURL);
 	Writer.WriteValue(TEXT("DocsURL"), DocsURL);
 	Writer.WriteValue(TEXT("MarketplaceURL"), MarketplaceURL);
+	Writer.WriteValue(TEXT("SupportURL"), SupportURL);
 
 	FModuleDescriptor::WriteArray(Writer, TEXT("Modules"), Modules);
 
diff --git a/Engine/Source/Runtime/Projects/Public/PluginDescriptor.h b/Engine/Source/Runtime/Projects/Public/PluginDescriptor.h
--- a/Engine/Source/Runtime/Projects/Public/PluginDescriptor.h
+++ b/Engine/Source/Runtime/Projects/Public/PluginDescriptor.h
@@ -61,6 +61,9 @@ struct PROJECTS_API FPluginDescriptor
 	/** Marketplace URL for this plugin. This URL will be embedded into projects that enable this plugin, so we can redirect to the marketplace if a user doesn't have it installed. */
 	FString MarketplaceURL;
 
+	/** Support URL/email for this plugin. */
+	FString SupportURL;
+
 	/** List of all modules associated with this plugin */
 	TArray<FModuleDescriptor> Modules;
 
