Ensured that Plugin Browser displays correct results after adding a new plugin

Added a delegate to the plugin browser module so that the plugin wizard can inform the plugin browser it's created a new one
Made the PendingEnablePlugins map private and put some of the logic using it into plugin browser module functions so that we can reuse it
Stopped opening the new plugin folder on creation as we don't expect users to do anything with it immediately

diff --git a/Engine/Plugins/Editor/PluginBrowser/Source/PluginBrowser/Private/PluginBrowserModule.cpp b/Engine/Plugins/Editor/PluginBrowser/Source/PluginBrowser/Private/PluginBrowserModule.cpp
--- a/Engine/Plugins/Editor/PluginBrowser/Source/PluginBrowser/Private/PluginBrowserModule.cpp
+++ b/Engine/Plugins/Editor/PluginBrowser/Source/PluginBrowser/Private/PluginBrowserModule.cpp
@@ -72,6 +72,30 @@ void FPluginBrowserModule::ShutdownModule()
 	IModularFeatures::Get().UnregisterModularFeature( EditorFeatures::PluginsEditor, this );
 }
 
+void FPluginBrowserModule::SetPluginPendingEnableState(const FString& PluginName, bool bCurrentlyEnabled, bool bPendingEnabled)
+{
+	if (bCurrentlyEnabled == bPendingEnabled)
+	{
+		PendingEnablePlugins.Remove(PluginName);
+	}
+	else
+	{
+		PendingEnablePlugins.FindOrAdd(PluginName) = bPendingEnabled;
+	}
+}
+
+bool FPluginBrowserModule::GetPluginPendingEnableState(const FString& PluginName) const
+{
+	check(PendingEnablePlugins.Contains(PluginName));
+
+	return PendingEnablePlugins[PluginName];
+}
+
+bool FPluginBrowserModule::HasPluginPendingEnable(const FString& PluginName) const
+{
+	return PendingEnablePlugins.Contains(PluginName);
+}
+
 TSharedRef<SDockTab> FPluginBrowserModule::HandleSpawnPluginCreatorTab(const FSpawnTabArgs& SpawnTabArgs)
 {
 	TSharedRef<SDockTab> ResultTab = SNew(SDockTab)
diff --git a/Engine/Plugins/Editor/PluginBrowser/Source/PluginBrowser/Private/PluginBrowserModule.h b/Engine/Plugins/Editor/PluginBrowser/Source/PluginBrowser/Private/PluginBrowserModule.h
--- a/Engine/Plugins/Editor/PluginBrowser/Source/PluginBrowser/Private/PluginBrowserModule.h
+++ b/Engine/Plugins/Editor/PluginBrowser/Source/PluginBrowser/Private/PluginBrowserModule.h
@@ -6,11 +6,11 @@
 #include "IPluginManager.h"
 #include "ModuleManager.h"
 
+DECLARE_MULTICAST_DELEGATE(FOnNewPluginCreated);
+
 class FPluginBrowserModule : public IPluginBrowser
 {
 public:
-	TMap<FString, bool> PendingEnablePlugins;
-
 	/** Accessor for the module interface */
 	static FPluginBrowserModule& Get()
 	{
@@ -21,6 +21,36 @@ public:
 	virtual void StartupModule() override;
 	virtual void ShutdownModule() override;
 
+	/** Gets a delegate so that you can register/unregister to receive callbacks when plugins are created */
+	FOnNewPluginCreated& OnNewPluginCreated() {return NewPluginCreatedDelegate;}
+
+	/** Broadcasts callback to notify registrants that a plugin has been created */
+	void BroadcastNewPluginCreated() const {NewPluginCreatedDelegate.Broadcast();}
+
+	/**
+	 * Sets whether a plugin is pending enable/disable
+	 * @param PluginName The name of the plugin
+	 * @param bCurrentlyEnabled The current state of this plugin, so that we can decide whether a change is no longer pending
+	 * @param bPendingEnabled Whether to set this plugin to pending enable or disable
+	 */
+	void SetPluginPendingEnableState(const FString& PluginName, bool bCurrentlyEnabled, bool bPendingEnabled);
+
+	/**
+	 * Gets whether a plugin is pending enable/disable
+	 * This should only be used when you know this is the case after using HasPluginPendingEnable
+	 * @param PluginName The name of the plugin
+	 */
+	bool GetPluginPendingEnableState(const FString& PluginName) const;
+
+	/** Whether there are any plugins pending enable/disable */
+	bool HasPluginsPendingEnable() const {return PendingEnablePlugins.Num() > 0;}
+
+	/**
+	 * Whether a specific plugin is pending enable/disable
+	 * @param PluginName The name of the plugin
+	 */
+	bool HasPluginPendingEnable(const FString& PluginName) const;
+
 	/** ID name for the plugins editor major tab */
 	static const FName PluginsEditorTabName;
 
@@ -28,5 +58,12 @@ public:
 	static const FName PluginCreatorTabName;
 
 private:
+	/** Called to spawn the plugin creator tab */
 	TSharedRef<SDockTab> HandleSpawnPluginCreatorTab(const FSpawnTabArgs& SpawnTabArgs);
+
+	/** List of plugins that are pending enable/disable */
+	TMap<FString, bool> PendingEnablePlugins;
+
+	/** Delegate called when a new plugin is created */
+	FOnNewPluginCreated NewPluginCreatedDelegate;
 };
diff --git a/Engine/Plugins/Editor/PluginBrowser/Source/PluginBrowser/Private/SNewPluginWizard.cpp b/Engine/Plugins/Editor/PluginBrowser/Source/PluginBrowser/Private/SNewPluginWizard.cpp
--- a/Engine/Plugins/Editor/PluginBrowser/Source/PluginBrowser/Private/SNewPluginWizard.cpp
+++ b/Engine/Plugins/Editor/PluginBrowser/Source/PluginBrowser/Private/SNewPluginWizard.cpp
@@ -10,6 +10,7 @@
 #include "SNotificationList.h"
 #include "NotificationManager.h"
 #include "GameProjectUtils.h"
+#include "PluginBrowserModule.h"
 
 DEFINE_LOG_CATEGORY(LogPluginWizard);
 
@@ -454,15 +455,22 @@ FReply SNewPluginWizard::OnCreatePluginClicked()
 
 	if (bSucceeded)
 	{
-		// Show the plugin folder in the file system
-		FPlatformProcess::ExploreFolder(*UPluginFilePath);
-
 		// Generate project files if we happen to be using a project file.
 		if (!FDesktopPlatformModule::Get()->GenerateProjectFiles(FPaths::RootDir(), FPaths::GetProjectFilePath(), GWarn))
 		{
 			PopErrorNotification(LOCTEXT("FailedToGenerateProjectFiles", "Failed to generate project files."));
 		}
 
+		// Notify that a new plugin has been created
+		FPluginBrowserModule& PluginBrowserModule = FPluginBrowserModule::Get();
+		PluginBrowserModule.BroadcastNewPluginCreated();
+
+		// Add game plugins to list of pending enable plugins 
+		if (!bIsEnginePlugin)
+		{
+			PluginBrowserModule.SetPluginPendingEnableState(AutoPluginName, false, true);
+		}
+
 		auto PinnedOwnerTab = OwnerTab.Pin();
 		if (PinnedOwnerTab.IsValid())
 		{
diff --git a/Engine/Plugins/Editor/PluginBrowser/Source/PluginBrowser/Private/SPluginBrowser.cpp b/Engine/Plugins/Editor/PluginBrowser/Source/PluginBrowser/Private/SPluginBrowser.cpp
--- a/Engine/Plugins/Editor/PluginBrowser/Source/PluginBrowser/Private/SPluginBrowser.cpp
+++ b/Engine/Plugins/Editor/PluginBrowser/Source/PluginBrowser/Private/SPluginBrowser.cpp
@@ -20,6 +20,8 @@ SPluginBrowser::~SPluginBrowser()
 	{
 		DirectoryWatcherModule.Get()->UnregisterDirectoryChangedCallback_Handle(Pair.Key, Pair.Value);
 	}
+
+	FPluginBrowserModule::Get().OnNewPluginCreated().RemoveAll(this);
 }
 
 void SPluginBrowser::Construct( const FArguments& Args )
@@ -43,6 +45,8 @@ void SPluginBrowser::Construct( const FArguments& Args )
 		}
 	}
 
+	FPluginBrowserModule::Get().OnNewPluginCreated().AddSP(this, &SPluginBrowser::OnNewPluginCreated);
+
 	RegisterActiveTimer (0.f, FWidgetActiveTimerDelegate::CreateSP (this, &SPluginBrowser::TriggerBreadcrumbRefresh));
 
 	struct Local
@@ -179,7 +183,7 @@ void SPluginBrowser::Construct( const FArguments& Args )
 
 EVisibility SPluginBrowser::HandleRestartEditorNoticeVisibility() const
 {
-	return (FPluginBrowserModule::Get().PendingEnablePlugins.Num() > 0)? EVisibility::Visible : EVisibility::Collapsed;
+	return FPluginBrowserModule::Get().HasPluginsPendingEnable() ? EVisibility::Visible : EVisibility::Collapsed;
 }
 
 
@@ -240,6 +244,15 @@ void SPluginBrowser::OnPluginDirectoryChanged(const TArray<struct FFileChangeDat
 	UpdatePluginsTimerHandle = RegisterActiveTimer(2.0f, FWidgetActiveTimerDelegate::CreateSP(this, &SPluginBrowser::UpdatePluginsTimerCallback));
 }
 
+void SPluginBrowser::OnNewPluginCreated()
+{
+	if (UpdatePluginsTimerHandle.IsValid())
+	{
+		UnRegisterActiveTimer(UpdatePluginsTimerHandle.ToSharedRef());
+	}
+	UpdatePluginsTimerHandle = RegisterActiveTimer(2.0f, FWidgetActiveTimerDelegate::CreateSP(this, &SPluginBrowser::UpdatePluginsTimerCallback));
+}
+
 EActiveTimerReturnType SPluginBrowser::UpdatePluginsTimerCallback(double InCurrentTime, float InDeltaTime)
 {
 	IPluginManager::Get().RefreshPluginsList();
diff --git a/Engine/Plugins/Editor/PluginBrowser/Source/PluginBrowser/Private/SPluginBrowser.h b/Engine/Plugins/Editor/PluginBrowser/Source/PluginBrowser/Private/SPluginBrowser.h
--- a/Engine/Plugins/Editor/PluginBrowser/Source/PluginBrowser/Private/SPluginBrowser.h
+++ b/Engine/Plugins/Editor/PluginBrowser/Source/PluginBrowser/Private/SPluginBrowser.h
@@ -42,6 +42,9 @@ private:
 	/** Called with notification that one of the plugin directories has changed */
 	void OnPluginDirectoryChanged(const TArray<struct FFileChangeData>& );
 
+	/** Called with notification that a new plugin has been created */
+	void OnNewPluginCreated();
+
 	/** Timer callback for when */
 	EActiveTimerReturnType UpdatePluginsTimerCallback(double InCurrentTime, float InDeltaTime);
 
diff --git a/Engine/Plugins/Editor/PluginBrowser/Source/PluginBrowser/Private/SPluginTile.cpp b/Engine/Plugins/Editor/PluginBrowser/Source/PluginBrowser/Private/SPluginTile.cpp
--- a/Engine/Plugins/Editor/PluginBrowser/Source/PluginBrowser/Private/SPluginTile.cpp
+++ b/Engine/Plugins/Editor/PluginBrowser/Source/PluginBrowser/Private/SPluginTile.cpp
@@ -359,9 +359,9 @@ void SPluginTile::RecreateWidgets()
 ECheckBoxState SPluginTile::IsPluginEnabled() const
 {
 	FPluginBrowserModule& PluginBrowserModule = FPluginBrowserModule::Get();
-	if(PluginBrowserModule.PendingEnablePlugins.Contains(Plugin->GetName()))
+	if(PluginBrowserModule.HasPluginPendingEnable(Plugin->GetName()))
 	{
-		return PluginBrowserModule.PendingEnablePlugins[Plugin->GetName()]? ECheckBoxState::Checked : ECheckBoxState::Unchecked;;
+		return PluginBrowserModule.GetPluginPendingEnableState(Plugin->GetName()) ? ECheckBoxState::Checked : ECheckBoxState::Unchecked;;
 	}
 	else
 	{
@@ -403,15 +403,7 @@ void SPluginTile::OnEnablePluginCheckboxChanged(ECheckBoxState NewCheckedState)
 		}
 		else
 		{
-			FPluginBrowserModule& PluginBrowserModule = FPluginBrowserModule::Get();
-			if(Plugin->IsEnabled() == bNewEnabledState)
-			{
-				PluginBrowserModule.PendingEnablePlugins.Remove(Plugin->GetName());
-			}
-			else
-			{
-				PluginBrowserModule.PendingEnablePlugins.FindOrAdd(Plugin->GetName()) = bNewEnabledState;
-			}
+			FPluginBrowserModule::Get().SetPluginPendingEnableState(Plugin->GetName(), Plugin->IsEnabled(), bNewEnabledState);
 		}
 	}
 
