Changed Plugin Creator to allow plugins to be created in sub directories

Copied SFilePath from SNewProjectWizard and added some new features to allow for error status messages on the edit boxes
Added Plugin Creator Background to the plugin style as well as some more padding for the elements in it
Removed unnecessary code creating each individual folder for new plugins as each file creation/copy creates the new directory internally

diff --git a/Engine/Plugins/Editor/PluginBrowser/Source/PluginBrowser/PluginBrowser.Build.cs b/Engine/Plugins/Editor/PluginBrowser/Source/PluginBrowser/PluginBrowser.Build.cs
--- a/Engine/Plugins/Editor/PluginBrowser/Source/PluginBrowser/PluginBrowser.Build.cs
+++ b/Engine/Plugins/Editor/PluginBrowser/Source/PluginBrowser/PluginBrowser.Build.cs
@@ -26,6 +26,7 @@ namespace UnrealBuildTool.Rules
 					"SharedSettingsWidgets",
 					"DirectoryWatcher",
 					"GameProjectGeneration",
+					"MainFrame"
 				}
 			);
 
diff --git a/Engine/Plugins/Editor/PluginBrowser/Source/PluginBrowser/Private/PluginStyle.cpp b/Engine/Plugins/Editor/PluginBrowser/Source/PluginBrowser/Private/PluginStyle.cpp
--- a/Engine/Plugins/Editor/PluginBrowser/Source/PluginBrowser/Private/PluginStyle.cpp
+++ b/Engine/Plugins/Editor/PluginBrowser/Source/PluginBrowser/Private/PluginStyle.cpp
@@ -181,6 +181,14 @@ void FPluginStyle::Initialize()
 		}
 	}
 
+	// Plugin Creator
+	{
+		const float PaddingAmount = 5.0f;
+		StyleSet->Set( "PluginCreator.Padding", PaddingAmount);
+
+		StyleSet->Set( "PluginCreator.Background", new BOX_BRUSH( "AppTabContentArea", FMargin(4/16.0f) ) );
+	}
+
 	FSlateStyleRegistry::RegisterSlateStyle( *StyleSet.Get() );
 };
 
diff --git a/Engine/Plugins/Editor/PluginBrowser/Source/PluginBrowser/Private/SFilePathBlock.cpp b/Engine/Plugins/Editor/PluginBrowser/Source/PluginBrowser/Private/SFilePathBlock.cpp
new file mode 100644
--- /dev/null
+++ b/Engine/Plugins/Editor/PluginBrowser/Source/PluginBrowser/Private/SFilePathBlock.cpp
@@ -0,0 +1,141 @@
+// Copyright 1998-2015 Epic Games, Inc. All Rights Reserved.
+
+#include "PluginBrowserPrivatePCH.h"
+#include "SFilePathBlock.h"
+
+#define LOCTEXT_NAMESPACE "FilePathBlock"
+
+void SFilePathBlock::Construct(const FArguments& InArgs)
+{
+	ChildSlot
+	[
+		SNew(SGridPanel)
+		.FillColumn(0, 2.f)
+		.FillColumn(1, 1.f)
+
+		// Folder input
+		+ SGridPanel::Slot(0, 0)
+		[
+			SNew(SOverlay)
+
+			+ SOverlay::Slot()
+			[
+				SAssignNew(FolderPathTextBox, SEditableTextBox)
+				.Text(InArgs._FolderPath)
+				// Large right hand padding to make room for the browse button
+				.Padding(FMargin(5.f, 3.f, 25.f, 3.f))
+				.OnTextChanged(InArgs._OnFolderChanged)
+				.OnTextCommitted(InArgs._OnFolderCommitted)
+			]
+
+			+ SOverlay::Slot()
+			.HAlign(HAlign_Right)
+			[
+				SNew(SButton)
+				.ButtonStyle(FEditorStyle::Get(), "FilePath.FolderButton")
+				.ContentPadding(FMargin(4.f, 0.f))
+				.OnClicked(InArgs._OnBrowseForFolder)
+				.ToolTipText(LOCTEXT("BrowseForFolder", "Browse for a folder"))
+				.Text(LOCTEXT("...", "..."))
+			]
+		]
+
+		// Folder label
+		+ SGridPanel::Slot(0, 1)
+		[
+			SNew(SOverlay)
+
+			+ SOverlay::Slot()
+			.VAlign(VAlign_Center)
+			[
+				SNew(SBox)
+				.HeightOverride(3)
+				[
+					SNew(SBorder)
+					.BorderImage(FEditorStyle::GetBrush("FilePath.GroupIndicator"))
+					.BorderBackgroundColor(FLinearColor(1.f, 1.f, 1.f, 0.5f))
+					.Padding(FMargin(150.f, 0.f))
+				]
+			]
+
+			+ SOverlay::Slot()
+			.HAlign(HAlign_Center)
+			.VAlign(VAlign_Center)
+			[
+				SNew(SBorder)
+				.Padding(5.f)
+				.BorderImage(InArgs._LabelBackgroundBrush)
+				.BorderBackgroundColor(InArgs._LabelBackgroundColor)
+				[
+					SNew(STextBlock)
+					.Text(LOCTEXT("Folder", "Folder"))
+				]
+			]
+		]
+
+		// Name input
+		+ SGridPanel::Slot(1, 0)
+		.Padding(FMargin(5.f, 0.f, 0.f, 0.f))
+		.VAlign(VAlign_Center)
+		[
+			SAssignNew(NameTextBox, SEditableTextBox)
+			.Text(InArgs._Name)
+			.Padding(FMargin(5.f, 3.f))
+			.HintText(InArgs._NameHint)
+			.OnTextChanged(InArgs._OnNameChanged)
+			.OnTextCommitted(InArgs._OnNameCommitted)
+		]
+
+		// Name label
+		+ SGridPanel::Slot(1, 1)
+		.Padding(FMargin(5.f, 0.f, 0.f, 0.f))
+		[
+			SNew(SOverlay)
+
+			+ SOverlay::Slot()
+			.VAlign(VAlign_Center)
+			[
+				SNew(SBox)
+				.HeightOverride(3)
+				[
+					SNew(SBorder)
+					.BorderImage(FEditorStyle::GetBrush("FilePath.GroupIndicator"))
+					.BorderBackgroundColor(FLinearColor(1.f, 1.f, 1.f, 0.5f))
+					.Padding(FMargin(75.f, 0.f))
+				]
+			]
+
+			+ SOverlay::Slot()
+			.HAlign(HAlign_Center)
+			.VAlign(VAlign_Center)
+			[
+				SNew(SBorder)
+				.Padding(5.f)
+				.BorderImage(InArgs._LabelBackgroundBrush)
+				.BorderBackgroundColor(InArgs._LabelBackgroundColor)
+				[
+					SNew(STextBlock)
+					.Text(LOCTEXT("Name", "Name"))
+				]
+			]
+		]
+	];
+}
+
+void SFilePathBlock::SetFolderPathError(const FText& ErrorText)
+{
+	if (FolderPathTextBox.IsValid())
+	{
+		FolderPathTextBox->SetError(ErrorText);
+	}
+}
+
+void SFilePathBlock::SetNameError(const FText& ErrorText)
+{
+	if (NameTextBox.IsValid())
+	{
+		NameTextBox->SetError(ErrorText);
+	}
+}
+
+#undef LOCTEXT_NAMESPACE
diff --git a/Engine/Plugins/Editor/PluginBrowser/Source/PluginBrowser/Private/SFilePathBlock.h b/Engine/Plugins/Editor/PluginBrowser/Source/PluginBrowser/Private/SFilePathBlock.h
new file mode 100644
--- /dev/null
+++ b/Engine/Plugins/Editor/PluginBrowser/Source/PluginBrowser/Private/SFilePathBlock.h
@@ -0,0 +1,56 @@
+// Copyright 1998-2015 Epic Games, Inc. All Rights Reserved.
+
+#pragma once
+
+/**
+ * Simple widget used to display a folder path, and a name of a file:
+ * __________________________  ____________________
+ * | [REDACTED_USER_HOME]    |  | SomeFile.txt     |
+ * |-------- Folder --------|  |------ Name ------|
+ */
+class SFilePathBlock : public SCompoundWidget
+{
+public:
+
+	SLATE_BEGIN_ARGS( SFilePathBlock )
+		: _LabelBackgroundColor(FLinearColor::Black)
+		, _LabelBackgroundBrush(FEditorStyle::GetBrush("WhiteBrush"))
+	{}
+		/** Attribute specifying the text to display in the folder input */
+		SLATE_ATTRIBUTE(FText, FolderPath)
+
+		/** Attribute specifying the text to display in the name input */
+		SLATE_ATTRIBUTE(FText, Name)
+
+		/** Hint name that appears when there is no text in the name box */
+		SLATE_ATTRIBUTE(FText, NameHint)
+
+		/** Background label tint for the folder/name labels */
+		SLATE_ATTRIBUTE(FSlateColor, LabelBackgroundColor)
+
+		/** Background label brush for the folder/name labels */
+		SLATE_ATTRIBUTE(const FSlateBrush*, LabelBackgroundBrush)
+
+		/** Event that is triggered when the browser for folder button is clicked */
+		SLATE_EVENT(FOnClicked, OnBrowseForFolder)
+
+		/** Events for when the name field is manipulated */
+		SLATE_EVENT(FOnTextChanged, OnNameChanged)
+		SLATE_EVENT(FOnTextCommitted, OnNameCommitted)
+		
+		/** Events for when the folder field is manipulated */
+		SLATE_EVENT(FOnTextChanged, OnFolderChanged)
+		SLATE_EVENT(FOnTextCommitted, OnFolderCommitted)
+
+	SLATE_END_ARGS()
+
+	/** Constructs this widget with InArgs */
+	void Construct( const FArguments& InArgs );
+
+	void SetFolderPathError(const FText& ErrorText);
+	void SetNameError(const FText& ErrorText);
+
+private:
+	TSharedPtr<SEditableTextBox> FolderPathTextBox;
+	TSharedPtr<SEditableTextBox> NameTextBox;
+};
\ No newline at end of file
diff --git a/Engine/Plugins/Editor/PluginBrowser/Source/PluginBrowser/Private/SNewPluginWizard.cpp b/Engine/Plugins/Editor/PluginBrowser/Source/PluginBrowser/Private/SNewPluginWizard.cpp
--- a/Engine/Plugins/Editor/PluginBrowser/Source/PluginBrowser/Private/SNewPluginWizard.cpp
+++ b/Engine/Plugins/Editor/PluginBrowser/Source/PluginBrowser/Private/SNewPluginWizard.cpp
@@ -11,13 +11,16 @@
 #include "NotificationManager.h"
 #include "GameProjectUtils.h"
 #include "PluginBrowserModule.h"
+#include "SFilePathBlock.h"
+#include "IMainFrameModule.h"
 
 DEFINE_LOG_CATEGORY(LogPluginWizard);
 
 #define LOCTEXT_NAMESPACE "NewPluginWizard"
 
 SNewPluginWizard::SNewPluginWizard()
-	: bIsPluginNameValid(false)
+	: bIsPluginPathValid(false)
+	, bIsPluginNameValid(false)
 	, bIsEnginePlugin(false)
 {
 	const FText BlankTemplateName = LOCTEXT("BlankLabel", "Blank");
@@ -36,13 +39,22 @@ SNewPluginWizard::SNewPluginWizard()
 	Templates.Add(MakeShareable(new FPluginTemplateDescription(AdvancedTemplateName, AdvancedDescription, TEXT("Advanced"), true, false, false)));
 	Templates.Add(MakeShareable(new FPluginTemplateDescription(BlueprintLibTemplateName, BlueprintLibDescription, TEXT("Blank"), false, false, true)));
 	Templates.Add(MakeShareable(new FPluginTemplateDescription(EditorModeTemplateName, EditorModeDescription, TEXT("Blank"), false, true, false)));
+
+	AbsoluteGamePluginPath = IFileManager::Get().ConvertToAbsolutePathForExternalAppForWrite(*FPaths::GamePluginsDir());
+	FPaths::MakePlatformFilename(AbsoluteGamePluginPath);
+	AbsoluteEnginePluginPath = IFileManager::Get().ConvertToAbsolutePathForExternalAppForWrite(*FPaths::EnginePluginsDir());
+	FPaths::MakePlatformFilename(AbsoluteEnginePluginPath);
 }
 
 void SNewPluginWizard::Construct(const FArguments& Args, TSharedPtr<SDockTab> InOwnerTab)
 {
 	OwnerTab = InOwnerTab;
 
-	const float PaddingAmount = FPluginStyle::Get()->GetFloat("PluginTile.Padding");
+	LastBrowsePath = AbsoluteGamePluginPath;
+	PluginFolderPath = AbsoluteGamePluginPath;
+	bIsPluginPathValid = true;
+
+	const float PaddingAmount = FPluginStyle::Get()->GetFloat("PluginCreator.Padding");
 
 	TSharedRef<SVerticalBox> MainContent = SNew(SVerticalBox)
 	+SVerticalBox::Slot()
@@ -69,43 +81,17 @@ void SNewPluginWizard::Construct(const FArguments& Args, TSharedPtr<SDockTab> In
 	+SVerticalBox::Slot()
 	.AutoHeight()
 	.Padding(PaddingAmount)
+	.HAlign(HAlign_Center)
 	[
-		SNew(SGridPanel)
-		.FillColumn(1, 1.0f)
-
-		+SGridPanel::Slot(0, 0)
-		.VAlign(VAlign_Center)
-		.Padding(PaddingAmount)
-		[
-			SNew(STextBlock)
-			.Text(LOCTEXT("NameLabel", "Name"))
-			.TextStyle(FPluginStyle::Get(), "PluginTile.NameText")
-		]
-		+SGridPanel::Slot(1, 0)
-		.VAlign(VAlign_Center)
-		.Padding(PaddingAmount)
-		[
-			// name of new plugin
-			SAssignNew(PluginNameTextBox, SEditableTextBox)
-			.HintText(LOCTEXT("PluginNameTextHint", "Plugin name"))
-			.OnTextChanged(this, &SNewPluginWizard::OnPluginNameTextChanged)
-		]
-		+SGridPanel::Slot(0, 1)
-		.VAlign(VAlign_Center)
-		.Padding(PaddingAmount)
-		[
-			SNew(STextBlock)
-			.Text(LOCTEXT("PathLabel", "Path"))
-			.TextStyle(FPluginStyle::Get(), "PluginTile.NameText")
-		]
-		+ SGridPanel::Slot(1, 1)
-		.VAlign(VAlign_Center)
-		.Padding(PaddingAmount)
-		[
-			// path for new plugin
-			SNew(STextBlock)
-			.Text(this, &SNewPluginWizard::GetPluginDestinationPath)
-		]
+		SAssignNew(FilePathBlock, SFilePathBlock)
+		.OnBrowseForFolder(this, &SNewPluginWizard::OnBrowseButtonClicked)
+		.LabelBackgroundBrush(FPluginStyle::Get()->GetBrush("PluginCreator.Background"))
+		.LabelBackgroundColor(FLinearColor::White)
+		.FolderPath(this, &SNewPluginWizard::GetPluginDestinationPath)
+		.Name(this, &SNewPluginWizard::GetCurrentPluginName)
+		.NameHint(LOCTEXT("PluginNameTextHint", "Plugin name"))
+		.OnFolderChanged(this, &SNewPluginWizard::OnFolderPathTextChanged)
+		.OnNameChanged(this, &SNewPluginWizard::OnPluginNameTextChanged)
 	];
 
 	// Don't show the option to make an engine plugin in installed builds
@@ -117,6 +103,8 @@ void SNewPluginWizard::Construct(const FArguments& Args, TSharedPtr<SDockTab> In
 		[
 			SNew(SBox)
 			.Padding(PaddingAmount)
+			.HAlign(HAlign_Left)
+			.VAlign(VAlign_Center)
 			[
 				SNew(SCheckBox)
 				.OnCheckStateChanged(this, &SNewPluginWizard::OnEnginePluginCheckboxChanged)
@@ -229,74 +217,189 @@ void SNewPluginWizard::OnTemplateSelectionChanged( TSharedPtr<FPluginTemplateDes
 	CurrentTemplate = InItem;
 }
 
+void SNewPluginWizard::OnFolderPathTextChanged(const FText& InText)
+{
+	PluginFolderPath = InText.ToString();
+	FPaths::MakePlatformFilename(PluginFolderPath);
+	ValidateFullPluginPath();
+}
+
 void SNewPluginWizard::OnPluginNameTextChanged(const FText& InText)
 {
-	// Early exit if text is empty
-	if (InText.IsEmpty())
+	PluginNameText = InText;
+	ValidateFullPluginPath();
+}
+
+FReply SNewPluginWizard::OnBrowseButtonClicked()
+{
+	IDesktopPlatform* DesktopPlatform = FDesktopPlatformModule::Get();
+	if (DesktopPlatform)
 	{
-		bIsPluginNameValid = false;
-		PluginNameText = FText();
-		PluginNameTextBox->SetError(FText::GetEmpty());
-		return;
+		void* ParentWindowWindowHandle = NULL;
+
+		IMainFrameModule& MainFrameModule = FModuleManager::LoadModuleChecked<IMainFrameModule>(TEXT("MainFrame"));
+		const TSharedPtr<SWindow>& MainFrameParentWindow = MainFrameModule.GetParentWindow();
+		if (MainFrameParentWindow.IsValid() && MainFrameParentWindow->GetNativeWindow().IsValid())
+		{
+			ParentWindowWindowHandle = MainFrameParentWindow->GetNativeWindow()->GetOSWindowHandle();
+		}
+
+		FString FolderName;
+		const FString Title = LOCTEXT("NewPluginBrowseTitle", "Choose a plugin location").ToString();
+		const bool bFolderSelected = DesktopPlatform->OpenDirectoryDialog(
+			ParentWindowWindowHandle,
+			Title,
+			LastBrowsePath,
+			FolderName
+			);
+
+		if (bFolderSelected)
+		{
+			LastBrowsePath = FolderName;
+			OnFolderPathTextChanged(FText::FromString(FolderName));
+		}
+	}
+
+	return FReply::Handled();
+}
+
+void SNewPluginWizard::ValidateFullPluginPath()
+{
+	// Check for issues with path
+	bIsPluginPathValid = false;
+	bool bIsNewPathValid = true;
+	FText FolderPathError;
+
+	if (!FPaths::ValidatePath(GetPluginDestinationPath().ToString(), &FolderPathError))
+	{
+		bIsNewPathValid = false;
+	}
+
+	if (bIsNewPathValid)
+	{
+		bool bFoundValidPath = false;
+		FString AbsolutePath = IFileManager::Get().ConvertToAbsolutePathForExternalAppForWrite(*GetPluginDestinationPath().ToString());
+		FPaths::MakePlatformFilename(AbsolutePath);
+
+		if (AbsolutePath.StartsWith(AbsoluteGamePluginPath))
+		{
+			bFoundValidPath = true;
+			bIsEnginePlugin = false;
+		}
+
+		if (!bFoundValidPath && !FApp::IsEngineInstalled())
+		{
+			if (AbsolutePath.StartsWith(AbsoluteEnginePluginPath))
+			{
+				bFoundValidPath = true;
+				bIsEnginePlugin = true;
+			}
+		}
+
+		bIsNewPathValid = bFoundValidPath;
+		if (!bFoundValidPath)
+		{
+			if (FApp::IsEngineInstalled())
+			{
+				FolderPathError = LOCTEXT("InstalledPluginFolderPathError", "Plugins can only be created within your Project's Plugins folder");
+			}
+			else
+			{
+				FolderPathError = LOCTEXT("PluginFolderPathError", "Plugins can only be created within the Engine Plugins folder or your Project's Plugins folder");
+			}
+		}
+	}
+
+	bIsPluginPathValid = bIsNewPathValid;
+	FilePathBlock->SetFolderPathError(FolderPathError);
+
+	// Check for issues with name
+	bIsPluginNameValid = false;
+	bool bIsNewNameValid = true;
+	FText PluginNameError;
+
+	// Fail silently if text is empty
+	if (GetCurrentPluginName().IsEmpty())
+	{
+		bIsNewNameValid = false;
 	}
 
 	// Don't allow commas, dots, etc...
 	FString IllegalCharacters;
-	if (!GameProjectUtils::NameContainsOnlyLegalCharacters(InText.ToString(), IllegalCharacters))
+	if (bIsNewNameValid && !GameProjectUtils::NameContainsOnlyLegalCharacters(GetCurrentPluginName().ToString(), IllegalCharacters))
 	{
 		FFormatNamedArguments Args;
 		Args.Add(TEXT("IllegalCharacters"), FText::FromString(IllegalCharacters));
-		FText ErrorText = FText::Format(LOCTEXT("WrongPluginNameErrorText", "Plugin name cannot contain illegal characters like: \"{IllegalCharacters}\""), Args);
-
-		PluginNameTextBox->SetError(ErrorText);
-		bIsPluginNameValid = false;
-		return;
+		PluginNameError = FText::Format(LOCTEXT("WrongPluginNameErrorText", "Plugin name cannot contain illegal characters like: \"{IllegalCharacters}\""), Args);
+		bIsNewNameValid = false;
 	}
 
-	const FString TestPluginName = InText.ToString();
-
-	// Check to see if a a compiled plugin with this name exists (at any path)
-	const TArray< FPluginStatus > Plugins = IPluginManager::Get().QueryStatusForAllPlugins();
-	for (auto PluginIt(Plugins.CreateConstIterator()); PluginIt; ++PluginIt)
+	if (bIsNewNameValid)
 	{
-		const auto& PluginStatus = *PluginIt;
+		const FString& TestPluginName = GetCurrentPluginName().ToString();
 
-		if (PluginStatus.Name == TestPluginName)
+		// Check to see if a a compiled plugin with this name exists (at any path)
+		const TArray< FPluginStatus > Plugins = IPluginManager::Get().QueryStatusForAllPlugins();
+		for (auto PluginIt(Plugins.CreateConstIterator()); PluginIt; ++PluginIt)
 		{
-			PluginNameTextBox->SetError(LOCTEXT("PluginNameExistsErrorText", "A plugin with this name already exists!"));
-			bIsPluginNameValid = false;
-
-			return;
+			const auto& PluginStatus = *PluginIt;
+
+			if (PluginStatus.Name == TestPluginName)
+			{
+				PluginNameError = LOCTEXT("PluginNameExistsErrorText", "A plugin with this name already exists!");
+				bIsNewNameValid = false;
+				break;
+			}
 		}
 	}
 
 	// Check to see if a .uplugin exists at this path (in case there is an uncompiled or disabled plugin)
-	const FString TestPluginPath = GetPluginDestinationPath().ToString() / TestPluginName / (TestPluginName + TEXT(".uplugin"));
-	if (FPlatformFileManager::Get().GetPlatformFile().FileExists(*TestPluginPath))
+	if (bIsNewNameValid)
 	{
-		PluginNameTextBox->SetError(LOCTEXT("PluginPathExistsErrorText", "A plugin already exists at this path!"));
-		bIsPluginNameValid = false;
-
-		return;
+		const FString TestPluginPath = GetPluginFilenameWithPath();
+		if (!TestPluginPath.IsEmpty())
+		{
+			if (FPlatformFileManager::Get().GetPlatformFile().FileExists(*TestPluginPath))
+			{
+				PluginNameError = LOCTEXT("PluginPathExistsErrorText", "A plugin already exists at this path!");
+				bIsNewNameValid = false;
+			}
+		}
 	}
 
-	PluginNameText = InText;
-	PluginNameTextBox->SetError(FText::GetEmpty());
-	bIsPluginNameValid = true;
+	bIsPluginNameValid = bIsNewNameValid;
+	FilePathBlock->SetNameError(PluginNameError);
 }
 
 bool SNewPluginWizard::CanCreatePlugin() const
 {
-	return bIsPluginNameValid && CurrentTemplate.IsValid();
+	return bIsPluginPathValid && bIsPluginNameValid && CurrentTemplate.IsValid();
 }
 
 FText SNewPluginWizard::GetPluginDestinationPath() const
 {
-	const FString EnginePath = FPaths::EnginePluginsDir();
-	const FString GamePath = FPaths::GamePluginsDir();
+	return FText::FromString(PluginFolderPath);
+}
+
+FText SNewPluginWizard::GetCurrentPluginName() const
+{
+	return PluginNameText;
+}
 
-	return bIsEnginePlugin ? FText::FromString(IFileManager::Get().ConvertToAbsolutePathForExternalAppForWrite(*EnginePath))
-		: FText::FromString(IFileManager::Get().ConvertToAbsolutePathForExternalAppForWrite(*GamePath));
+FString SNewPluginWizard::GetPluginFilenameWithPath() const
+{
+	if (PluginFolderPath.IsEmpty() || PluginNameText.IsEmpty())
+	{
+		// Don't even try to assemble the path or else it may be relative to the binaries folder!
+		return TEXT("");
+	}
+	else
+	{
+		const FString& TestPluginName = PluginNameText.ToString();
+		FString TestPluginPath = PluginFolderPath / TestPluginName / (TestPluginName + TEXT(".uplugin"));
+		FPaths::MakePlatformFilename(TestPluginPath);
+		return TestPluginPath;
+	}
 }
 
 ECheckBoxState SNewPluginWizard::IsEnginePlugin() const
@@ -306,7 +409,21 @@ ECheckBoxState SNewPluginWizard::IsEnginePlugin() const
 
 void SNewPluginWizard::OnEnginePluginCheckboxChanged(ECheckBoxState NewCheckedState)
 {
-	bIsEnginePlugin = NewCheckedState == ECheckBoxState::Checked;
+	bool bNewEnginePluginState = NewCheckedState == ECheckBoxState::Checked;
+	if (bIsEnginePlugin != bNewEnginePluginState)
+	{
+		bIsEnginePlugin = bNewEnginePluginState;
+		if (bIsEnginePlugin)
+		{
+			PluginFolderPath = AbsoluteEnginePluginPath;
+		}
+		else
+		{
+			PluginFolderPath = AbsoluteGamePluginPath;
+		}
+		bIsPluginPathValid = true;
+		FilePathBlock->SetFolderPathError(FText::GetEmpty());
+	}
 }
 
 FReply SNewPluginWizard::OnCreatePluginClicked()
@@ -330,27 +447,17 @@ FReply SNewPluginWizard::OnCreatePluginClicked()
 
 	FText LocalFailReason;
 
-	// If we're creating a game plugin, GameDir\Plugins folder may not exist so we have to create it
-	if (!bIsEnginePlugin && !IFileManager::Get().DirectoryExists(*GetPluginDestinationPath().ToString()))
-	{
-		if (!MakeDirectory(*GetPluginDestinationPath().ToString()))
-		{
-			return FReply::Unhandled();
-		}
-	}
+	bool bSucceeded = true;
 
-	// Create main plugin dir
-	const FString PluginFolder = GetPluginDestinationPath().ToString() / AutoPluginName;
-	if (!MakeDirectory(PluginFolder))
-	{
-		return FReply::Unhandled();
-	}
+	// Save descriptor file as .uplugin file
+	const FString UPluginFilePath = GetPluginFilenameWithPath();
+	bSucceeded = bSucceeded && WritePluginDescriptor(AutoPluginName, UPluginFilePath);
 
-	bool bSucceeded = true;
+	// Main plugin dir
+	const FString PluginFolder = GetPluginDestinationPath().ToString() / AutoPluginName;
 
-	// Create resource folder
+	// Resource folder
 	const FString ResourcesFolder = PluginFolder / TEXT("Resources");
-	bSucceeded = bSucceeded && MakeDirectory(ResourcesFolder);
 
 	// Copy the icons
 	bSucceeded = bSucceeded && CopyFile(ResourcesFolder / TEXT("Icon128.png"), PluginEditorIconPath, /*inout*/ CreatedFiles);
@@ -361,23 +468,13 @@ FReply SNewPluginWizard::OnCreatePluginClicked()
 		bSucceeded = bSucceeded && CopyFile(ResourcesFolder / TEXT("ButtonIcon_40x.png"), PluginButtonIconPath, /*inout*/ CreatedFiles);
 	}
 
-	// Create source folder
+	// Source folder
 	const FString SourceFolder = PluginFolder / TEXT("Source");
-	bSucceeded = bSucceeded && MakeDirectory(SourceFolder);
-
-	// Save descriptor file as .uplugin file
-	const FString UPluginFilePath = PluginFolder / AutoPluginName + TEXT(".uplugin");
-	bSucceeded = bSucceeded && WritePluginDescriptor(AutoPluginName, UPluginFilePath);
 
-	// Create Source Folders
-	const FString PluginSourceFolder = SourceFolder / AutoPluginName;
-	bSucceeded = bSucceeded && MakeDirectory(PluginSourceFolder);
-
-	FString PrivateSourceFolder = PluginSourceFolder / TEXT("Private");
-	bSucceeded = bSucceeded && MakeDirectory(PrivateSourceFolder);
-
-	FString PublicSourceFolder = PluginSourceFolder / TEXT("Public");
-	bSucceeded = bSucceeded && MakeDirectory(PublicSourceFolder);
+	// Create Plugin Module Folders
+	const FString PluginModuleFolder = SourceFolder / AutoPluginName;
+	FString PrivateSourceFolder = PluginModuleFolder / TEXT("Private");
+	FString PublicSourceFolder = PluginModuleFolder / TEXT("Public");
 
 	TArray<FString> PrivateDependencyModuleNames;
 	PrivateDependencyModuleNames.Add("CoreUObject");
@@ -396,7 +493,7 @@ FReply SNewPluginWizard::OnCreatePluginClicked()
 	}
 
 	// Based on chosen template create build, and other source files
-	if (bSucceeded && !FPluginHelpers::CreatePluginBuildFile(PluginSourceFolder / AutoPluginName + TEXT(".Build.cs"), AutoPluginName, LocalFailReason, SelectedTemplate.OnDiskPath, PrivateDependencyModuleNames))
+	if (bSucceeded && !FPluginHelpers::CreatePluginBuildFile(PluginModuleFolder / AutoPluginName + TEXT(".Build.cs"), AutoPluginName, LocalFailReason, SelectedTemplate.OnDiskPath, PrivateDependencyModuleNames))
 	{
 		PopErrorNotification(FText::Format(LOCTEXT("FailedBuild", "Failed to create plugin build file. {0}"), LocalFailReason));
 		bSucceeded = false;
@@ -486,17 +583,6 @@ FReply SNewPluginWizard::OnCreatePluginClicked()
 	}
 }
 
-bool SNewPluginWizard::MakeDirectory(const FString& InPath)
-{
-	if (!IFileManager::Get().MakeDirectory(*InPath))
-	{
-		PopErrorNotification(FText::Format(LOCTEXT("CreateDirectoryFailed", "Failed to create directory '{0}'"), FText::AsCultureInvariant(InPath)));
-		return false;
-	}
-
-	return true;
-}
-
 bool SNewPluginWizard::CopyFile(const FString& DestinationFile, const FString& SourceFile, TArray<FString>& InOutCreatedFiles)
 {
 	if (IFileManager::Get().Copy(*DestinationFile, *SourceFile, /*bReplaceExisting=*/ false) != COPY_OK)
diff --git a/Engine/Plugins/Editor/PluginBrowser/Source/PluginBrowser/Private/SNewPluginWizard.h b/Engine/Plugins/Editor/PluginBrowser/Source/PluginBrowser/Private/SNewPluginWizard.h
--- a/Engine/Plugins/Editor/PluginBrowser/Source/PluginBrowser/Private/SNewPluginWizard.h
+++ b/Engine/Plugins/Editor/PluginBrowser/Source/PluginBrowser/Private/SNewPluginWizard.h
@@ -4,6 +4,8 @@
 
 DECLARE_LOG_CATEGORY_EXTERN(LogPluginWizard, Log, All);
 
+class SFilePathBlock;
+
 /**
  * Description of a plugin template
  */
@@ -73,12 +75,27 @@ private:
 	 */
 	void OnTemplateSelectionChanged( TSharedPtr<FPluginTemplateDescription> InItem, ESelectInfo::Type SelectInfo );
 
+	/**
+	 * Called when Folder Path textbox changes value
+	 * @param InText The new Plugin Folder Path text
+	 */
+	void OnFolderPathTextChanged(const FText& InText);
+
 	/**
 	 * Called when Plugin Name textbox changes value
 	 * @param InText The new Plugin name text
 	 */
 	void OnPluginNameTextChanged(const FText& InText);
 
+	/** Handler for when the Browse button is clicked */
+	FReply OnBrowseButtonClicked();
+
+	/**
+	 * Validates both the current path and plugin name as the final step in name
+	 * validation requires a valid path.
+	 */
+	void ValidateFullPluginPath();
+
 	/**
 	 * Whether we are currently able to create a plugin
 	 */
@@ -89,6 +106,12 @@ private:
 	 */
 	FText GetPluginDestinationPath() const;
 
+	/** Get the current name of the plugin */
+	FText GetCurrentPluginName() const;
+
+	/** Get the full path of the .plugin file we will create */
+	FString GetPluginFilenameWithPath() const;
+
 	/**
 	 * Whether we will create a plugin in the engine directory
 	 */
@@ -106,12 +129,6 @@ private:
 	 */
 	FReply OnCreatePluginClicked();
 
-	/**
-	 * Helper function to create directory under given path
-	 * @param InPath - path to a directory that you want to create
-	 */
-	bool MakeDirectory(const FString& InPath);
-
 	/**
 	 * Copies a file and adds to a list of created files
 	 * @param DestinationFile Where the file will be copied
@@ -148,11 +165,26 @@ private:
 	/** Currently selected template */
 	TSharedPtr<FPluginTemplateDescription> CurrentTemplate;
 
+	/** Absolute path to game plugins directory so we don't have to convert it repeatedly */
+	FString AbsoluteGamePluginPath;
+
+	/** Absolute path to engine plugins directory so we don't have to convert it repeatedly */
+	FString AbsoluteEnginePluginPath;
+
+	/** Last Path used to browse, so that we know it will open dialog */
+	FString LastBrowsePath;
+
+	/** Path where you want to create the plugin*/
+	FString PluginFolderPath;
+
 	/** Name of the plugin you want to create*/
 	FText PluginNameText;
 
-	/** Textbox widget that user will put plugin name in */
-	TSharedPtr<SEditableTextBox> PluginNameTextBox;
+	/** File Path widget that user will choose plugin location and name with */
+	TSharedPtr<SFilePathBlock> FilePathBlock;
+
+	/** Whether the path of the plugin entered is currently valid */
+	bool bIsPluginPathValid;
 
 	/** Whether the name of the plugin entered is currently valid */
 	bool bIsPluginNameValid;
