Reworked plugin templates so they are laid out with desired folder structure

Removed code that relied on specific .template files existing that are then opened and modified to create new files
Added individual folders for all templates so that they don't share any code

diff --git a/Engine/Plugins/Editor/PluginBrowser/Source/PluginBrowser/Private/PluginHelpers.cpp b/Engine/Plugins/Editor/PluginBrowser/Source/PluginBrowser/Private/PluginHelpers.cpp
--- a/Engine/Plugins/Editor/PluginBrowser/Source/PluginBrowser/Private/PluginHelpers.cpp
+++ b/Engine/Plugins/Editor/PluginBrowser/Source/PluginBrowser/Private/PluginHelpers.cpp
@@ -7,376 +7,6 @@
 
 #define LOCTEXT_NAMESPACE "PluginHelpers"
 
-bool FPluginHelpers::ReadTemplateFile(const FString& TemplateFileName, FString& OutFileContents, FText& OutFailReason)
-{
-	const FString FullFileName = IPluginManager::Get().FindPlugin(TEXT("PluginBrowser"))->GetBaseDir() / TEXT("Templates") / TemplateFileName;
-	if (FFileHelper::LoadFileToString(OutFileContents, *FullFileName))
-	{
-		return true;
-	}
-
-	FFormatNamedArguments Args;
-	Args.Add(TEXT("FullFileName"), FText::FromString(FullFileName));
-	OutFailReason = FText::Format(LOCTEXT("FailedToReadTemplateFile", "Failed to read template file \"{FullFileName}\""), Args);
-	return false;
-}
-
-bool FPluginHelpers::CreatePluginBuildFile(const FString& NewBuildFileName, const FString& PluginName, FText& OutFailReason, FString TemplateType, TArray<FString> InPrivateDependencyModuleNames)
-{
-	FString Template;
-
-	FString TemplateFileName = TEXT("PluginModule.Build.cs.template");
-
-	if (!ReadTemplateFile(TemplateType.IsEmpty() ? TemplateFileName : TemplateType / TemplateFileName, Template, OutFailReason))
-	{
-		return false;
-	}
-
-	TArray<FString> PublicIncludePathsNames;
-	TArray<FString> PrivateIncludePathsNames;
-	TArray<FString> PublicDependencyModuleNames;
-	TArray<FString> PrivateDependencyModuleNames;
-	TArray<FString> DynamicallyLoadedModuleNames;
-
-	if (InPrivateDependencyModuleNames.Num() == 0)
-	{
-		PrivateDependencyModuleNames.Add("CoreUObject");
-		PrivateDependencyModuleNames.Add("Engine");
-		PrivateDependencyModuleNames.Add("Slate");
-		PrivateDependencyModuleNames.Add("SlateCore");
-	}
-	else
-	{
-		PrivateDependencyModuleNames = InPrivateDependencyModuleNames;
-	}
-	
-	FString FinalOutput = Template.Replace(TEXT("%PUBLIC_INCLUDE_PATHS_NAMES%"), *GameProjectUtils::MakeCommaDelimitedList(PublicIncludePathsNames), ESearchCase::CaseSensitive);
-	FinalOutput = FinalOutput.Replace(TEXT("%PRIVATE_INCLUDE_PATHS_NAMES%"), *GameProjectUtils::MakeCommaDelimitedList(PrivateIncludePathsNames), ESearchCase::CaseSensitive);
-	FinalOutput = FinalOutput.Replace(TEXT("%PUBLIC_DEPENDENCY_MODULE_NAMES%"), *GameProjectUtils::MakeCommaDelimitedList(PublicDependencyModuleNames), ESearchCase::CaseSensitive);
-	FinalOutput = FinalOutput.Replace(TEXT("%PRIVATE_DEPENDENCY_MODULE_NAMES%"), *GameProjectUtils::MakeCommaDelimitedList(PrivateDependencyModuleNames), ESearchCase::CaseSensitive);
-	FinalOutput = FinalOutput.Replace(TEXT("%DYNAMICALLY_LOADED_MODULES_NAMES%"), *GameProjectUtils::MakeCommaDelimitedList(DynamicallyLoadedModuleNames), ESearchCase::CaseSensitive);
-	FinalOutput = FinalOutput.Replace(TEXT("%MODULE_NAME%"), *PluginName, ESearchCase::CaseSensitive);
-
-	return GameProjectUtils::WriteOutputFile(NewBuildFileName, FinalOutput, OutFailReason);
-}
-
-bool FPluginHelpers::CreatePluginHeaderFile(const FString& FolderPath, const FString& PluginName, FText& OutFailReason, FString TemplateType)
-{
-	FString Template;
-
-	FString TemplateFileName = TEXT("PluginModule.h.template");
-
-	if (!ReadTemplateFile(TemplateType.IsEmpty() ? TemplateFileName : TemplateType / TemplateFileName, Template, OutFailReason))
-	{
-		return false;
-	}
-
-	TArray<FString> PublicHeaderIncludes;
-
-	FString	FinalOutput = Template.Replace(TEXT("%PUBLIC_HEADER_INCLUDES%"), *GameProjectUtils::MakeIncludeList(PublicHeaderIncludes), ESearchCase::CaseSensitive);
-	FinalOutput = FinalOutput.Replace(TEXT("%PLUGIN_NAME%"), *PluginName, ESearchCase::CaseSensitive);
-
-	return GameProjectUtils::WriteOutputFile(FolderPath / PluginName + TEXT(".h"), FinalOutput, OutFailReason);
-}
-
-bool FPluginHelpers::CreatePluginCPPFile(const FString& FolderPath, const FString& PluginName, FText& OutFailReason, FString TemplateType, bool bMakeEditorMode, bool bUsesToolkits)
-{
-	FString Template;
-
-	FString TemplateFileName = TEXT("PluginModule.cpp.template");
-
-	if (!ReadTemplateFile(TemplateType.IsEmpty() ? TemplateFileName : TemplateType / TemplateFileName, Template, OutFailReason))
-	{
-		return false;
-	}
-
-	TArray<FString> PublicHeaderIncludes;
-	FString RegisterModeString;
-	FString UnRegisterModeString;
-
-	if (bMakeEditorMode)
-	{
-		FString EdModeName = PluginName + TEXT("EdMode");
-		PublicHeaderIncludes.Add(FString::Printf(TEXT("%s.h"), *(EdModeName)));
-		RegisterModeString = MakeRegisterEdModeString(EdModeName, bUsesToolkits);
-		UnRegisterModeString = MakeUnRegisterEdModeString(EdModeName);
-	}
-
-	FString	FinalOutput = Template.Replace(TEXT("%PUBLIC_HEADER_INCLUDES%"), *GameProjectUtils::MakeIncludeList(PublicHeaderIncludes), ESearchCase::CaseSensitive);
-	FinalOutput = FinalOutput.Replace(TEXT("%PLUGIN_NAME%"), *PluginName, ESearchCase::CaseSensitive);
-
-	FinalOutput = FinalOutput.Replace(TEXT("%REGISTER_EDMODE%"), *RegisterModeString, ESearchCase::CaseSensitive);
-	FinalOutput = FinalOutput.Replace(TEXT("%UNREGISTER_EDMODE%"), *UnRegisterModeString, ESearchCase::CaseSensitive);
-
-	return GameProjectUtils::WriteOutputFile(FolderPath / PluginName + TEXT(".cpp"), FinalOutput, OutFailReason);
-}
-
-bool FPluginHelpers::CreatePrivatePCHFile(const FString& FolderPath, const FString& PluginName, FText& OutFailReason, FString TemplateType)
-{
-	FString Template;
-	FString TemplateFileName = TEXT("PluginModulePCH.h.template");
-
-	if (!ReadTemplateFile(TemplateType.IsEmpty() ? TemplateFileName : TemplateType / TemplateFileName, Template, OutFailReason))
-	{
-		return false;
-	}
-
-	TArray<FString> PublicHeaderIncludes;
-
-	FString	FinalOutput = Template.Replace(TEXT("%PLUGIN_NAME%"), *PluginName, ESearchCase::CaseSensitive);
-
-	return GameProjectUtils::WriteOutputFile(FolderPath / PluginName + TEXT("PrivatePCH.h"), FinalOutput, OutFailReason);
-}
-
-bool FPluginHelpers::CreateBlueprintFunctionLibraryFiles(const FString& PrivateFolderPath, const FString& PublicFolderPath, const FString& PluginName, FText& OutFailReason)
-{
-	{
-		FString Template;
-		FString TemplateFileName = TEXT("BlueprintFunctionLibrary.h.template");
-
-		if (!ReadTemplateFile(TemplateFileName, Template, OutFailReason))
-		{
-			return false;
-		}
-
-		FString	FinalOutput = Template.Replace(TEXT("%PLUGIN_NAME%"), *PluginName, ESearchCase::CaseSensitive);
-
-		if (!GameProjectUtils::WriteOutputFile(PublicFolderPath / PluginName + TEXT("BPLibrary.h"), FinalOutput, OutFailReason))
-		{
-			return false;
-		}
-	}
-
-	{
-		FString Template;
-		FString TemplateFileName = TEXT("BlueprintFunctionLibrary.cpp.template");
-
-		if (!ReadTemplateFile(TemplateFileName, Template, OutFailReason))
-		{
-			return false;
-		}
-
-		FString	FinalOutput = Template.Replace(TEXT("%PLUGIN_NAME%"), *PluginName, ESearchCase::CaseSensitive);
-
-		if (!GameProjectUtils::WriteOutputFile(PrivateFolderPath / PluginName + TEXT("BPLibrary.cpp"), FinalOutput, OutFailReason))
-		{
-			return false;
-		}
-	}
-	
-
-	return true;
-}
-
-bool FPluginHelpers::CreatePluginStyleFiles(const FString& PrivateFolderPath, const FString& PublicFolderPath, const FString& PluginName, FText& OutFailReason, FString TemplateType)
-{
-	{
-		// Make Style.h file first
-
-		FString Template;
-		FString TemplateFileName = TEXT("PluginStyle.h.template");
-
-		if (!ReadTemplateFile(TemplateType.IsEmpty() ? TemplateFileName : TemplateType / TemplateFileName, Template, OutFailReason))
-		{
-			return false;
-		}
-
-		TArray<FString> PublicHeaderIncludes;
-
-		FString	FinalOutput = Template.Replace(TEXT("%PLUGIN_NAME%"), *PluginName, ESearchCase::CaseSensitive);
-
-		if (!GameProjectUtils::WriteOutputFile(PublicFolderPath / PluginName + TEXT("Style.h"), FinalOutput, OutFailReason))
-		{
-			return false;
-		}
-	}
-
-	{
-		// Now .cpp file
-
-		FString Template;
-
-		FString TemplateFileName = TEXT("PluginStyle.cpp.template");
-
-		if (!ReadTemplateFile(TemplateType.IsEmpty() ? TemplateFileName : TemplateType / TemplateFileName, Template, OutFailReason))
-		{
-			return false;
-		}
-
-		TArray<FString> PublicHeaderIncludes;
-
-		FString	FinalOutput = Template.Replace(TEXT("%PLUGIN_NAME%"), *PluginName, ESearchCase::CaseSensitive);
-
-		if (!GameProjectUtils::WriteOutputFile(PrivateFolderPath / PluginName + TEXT("Style.cpp"), FinalOutput, OutFailReason))
-		{
-			return false;
-		}
-	}
-
-	return true;
-	
-}
-
-bool FPluginHelpers::CreatePluginEdModeFiles(const FString& PrivateFolderPath, const FString& PublicFolderPath, const FString& PluginName, bool bUsesToolkits, bool bIncludeSampleUI, FText& OutFailReason, FString TemplateType /*= FString("")*/)
-{
-	FString EdModeName = PluginName +  + TEXT("EdMode");
-
-	{
-		// Make EdMode.h file first
-
-		FString Template;
-		FString TemplateFileName = TEXT("EdMode.h.template");
-
-		if (!ReadTemplateFile(TemplateType.IsEmpty() ? TemplateFileName : TemplateType / TemplateFileName, Template, OutFailReason))
-		{
-			return false;
-		}
-
-		FString	FinalOutput = Template.Replace(TEXT("%PLUGIN_NAME%"), *PluginName, ESearchCase::CaseSensitive);
-		FinalOutput = FinalOutput.Replace(TEXT("%MODE_NAME%"), *EdModeName, ESearchCase::CaseSensitive);
-
-		if (!GameProjectUtils::WriteOutputFile(PublicFolderPath / EdModeName + TEXT(".h"), FinalOutput, OutFailReason))
-		{
-			return false;
-		}
-	}
-
-	{
-		// Now .cpp file
-
-		FString Template;
-
-		FString TemplateFileName = TEXT("EdMode.cpp.template");
-
-		if (!ReadTemplateFile(TemplateType.IsEmpty() ? TemplateFileName : TemplateType / TemplateFileName, Template, OutFailReason))
-		{
-			return false;
-		}
-
-		FString UsesToolkit = bUsesToolkits ? TEXT("true") : TEXT("false");
-		FString BeginComment = bUsesToolkits ? TEXT("") : TEXT("/*");
-		FString EndComment = bUsesToolkits ? TEXT("") : TEXT("*/");
-
-		FString	FinalOutput = Template.Replace(TEXT("%PLUGIN_NAME%"), *PluginName, ESearchCase::CaseSensitive);
-		FinalOutput = FinalOutput.Replace(TEXT("%MODE_NAME%"), *EdModeName, ESearchCase::CaseSensitive);
-		FinalOutput = FinalOutput.Replace(TEXT("%BEGIN_BLOCK_COMMENT%"), *BeginComment, ESearchCase::CaseSensitive);
-		FinalOutput = FinalOutput.Replace(TEXT("%END_BLOCK_COMMENT%"), *EndComment, ESearchCase::CaseSensitive);
-		FinalOutput = FinalOutput.Replace(TEXT("%USES_TOOLKIT%"), *UsesToolkit, ESearchCase::CaseSensitive);
-
-
-		if (!GameProjectUtils::WriteOutputFile(PrivateFolderPath / EdModeName + TEXT(".cpp"), FinalOutput, OutFailReason))
-		{
-			return false;
-		}
-	}
-
-	// Make toolkit files if necessary
-	if (bUsesToolkits)
-	{
-		{
-			FString Template;
-			FString TemplateFileName = TEXT("Toolkit.h.template");
-
-			if (!ReadTemplateFile(TemplateType.IsEmpty() ? TemplateFileName : TemplateType / TemplateFileName, Template, OutFailReason))
-			{
-				return false;
-			}
-
-			FString	FinalOutput = Template.Replace(TEXT("%PLUGIN_NAME%"), *PluginName, ESearchCase::CaseSensitive);
-			FinalOutput = FinalOutput.Replace(TEXT("%MODE_NAME%"), *EdModeName, ESearchCase::CaseSensitive);
-
-			if (!GameProjectUtils::WriteOutputFile(PublicFolderPath / EdModeName + TEXT("Toolkit.h"), FinalOutput, OutFailReason))
-			{
-				return false;
-			}
-		}
-
-		{
-			FString Template;
-			FString TemplateFileName = TEXT("Toolkit.cpp.template");
-
-			if (!ReadTemplateFile(TemplateType.IsEmpty() ? TemplateFileName : TemplateType / TemplateFileName, Template, OutFailReason))
-			{
-				return false;
-			}
-
-			FString SampleUITemplate = TEXT("");
-
-			if (bIncludeSampleUI)
-			{
-				TemplateFileName = TEXT("ToolkitUI.template");
-
-				if (!ReadTemplateFile(TemplateFileName, SampleUITemplate, OutFailReason))
-				{
-					return false;
-				}
-			}
-
-			FString	FinalOutput = Template.Replace(TEXT("%PLUGIN_NAME%"), *PluginName, ESearchCase::CaseSensitive);
-			FinalOutput = FinalOutput.Replace(TEXT("%MODE_NAME%"), *EdModeName, ESearchCase::CaseSensitive);
-			FinalOutput = FinalOutput.Replace(TEXT("%TOOLKIT_SAMPLE_UI%"), *SampleUITemplate, ESearchCase::CaseSensitive);
-
-			if (!GameProjectUtils::WriteOutputFile(PrivateFolderPath / EdModeName + TEXT("Toolkit.cpp"), FinalOutput, OutFailReason))
-			{
-				return false;
-			}
-		}
-	}
-
-	return true;
-}
-
-bool FPluginHelpers::CreateCommandsFiles(const FString& FolderPath, const FString& PluginName, FText& OutFailReason, FString TemplateType)
-{
-	{
-		// Make Commands.h file first
-		FString Template;
-
-		FString TemplateFileName = TEXT("PluginCommands.h.template");
-
-		if (!ReadTemplateFile(TemplateType.IsEmpty() ? TemplateFileName : TemplateType / TemplateFileName, Template, OutFailReason))
-		{
-			return false;
-		}
-
-		TArray<FString> PublicHeaderIncludes;
-		FString ButtonLabel = TEXT("BUTTON LABEL");
-
-		FString	FinalOutput = Template.Replace(TEXT("%PLUGIN_NAME%"), *PluginName, ESearchCase::CaseSensitive);
-		FinalOutput = FinalOutput.Replace(TEXT("%PLUGIN_BUTTON_LABEL%"), *PluginName, ESearchCase::CaseSensitive);
-
-		if (!GameProjectUtils::WriteOutputFile(FolderPath / PluginName + TEXT("Commands.h"), FinalOutput, OutFailReason))
-		{
-			return false;
-		}
-	}
-
-	{
-		// Now .cpp file
-		FString Template;
-
-		FString TemplateFileName = TEXT("PluginCommands.cpp.template");
-
-		if (!ReadTemplateFile(TemplateType.IsEmpty() ? TemplateFileName : TemplateType / TemplateFileName, Template, OutFailReason))
-		{
-			return false;
-		}
-
-		TArray<FString> PublicHeaderIncludes;
-		FString ButtonLabel = TEXT("BUTTON LABEL");
-
-		FString	FinalOutput = Template.Replace(TEXT("%PLUGIN_NAME%"), *PluginName, ESearchCase::CaseSensitive);
-		FinalOutput = FinalOutput.Replace(TEXT("%PLUGIN_BUTTON_LABEL%"), *PluginName, ESearchCase::CaseSensitive);
-
-		if (!GameProjectUtils::WriteOutputFile(FolderPath / PluginName + TEXT("Commands.cpp"), FinalOutput, OutFailReason))
-		{
-			return false;
-		}
-	}
-
-	return true;
-}
-
 bool FPluginHelpers::CopyPluginTemplateFolder(const TCHAR* DestinationDirectory, const TCHAR* Source, const FString& PluginName)
 {
 	check(DestinationDirectory);
diff --git a/Engine/Plugins/Editor/PluginBrowser/Source/PluginBrowser/Private/PluginHelpers.h b/Engine/Plugins/Editor/PluginBrowser/Source/PluginBrowser/Private/PluginHelpers.h
--- a/Engine/Plugins/Editor/PluginBrowser/Source/PluginBrowser/Private/PluginHelpers.h
+++ b/Engine/Plugins/Editor/PluginBrowser/Source/PluginBrowser/Private/PluginHelpers.h
@@ -5,133 +5,6 @@
 class FPluginHelpers
 {
 public:
-	/**
-	 * Returns the contents of the specified template file.
-	 * @param TemplateFileName - name of the file with full path
-	 * @param OutFileContents - contents of given file
-	 * @param OutFailReason - fail reason text if something goes wrong
-	 * @return true if the file was loaded, false otherwise
-	 */
-	static bool ReadTemplateFile(const FString& TemplateFileName, FString& OutFileContents, FText& OutFailReason);
-
-	/**
-	 * Create a .build.cs file for a plugin.
-	 * @param NewBuildFileName - filename with full path
-	 * @param PluginName - name of the plugin this file belongs to
-	 * @param OutFailReason - fail reason text if something goes wrong
-	 * @param TemplateType - empty by default; can be Blank, Basic or Advanced
-	 * @param InPrivateDependencyModuleNames - Array of module names that should be included in this plugin. If it's empty, basic modules will be included (see definition of this function).
-	 * @return true if file was created, false otherwise
-	 */
-	static bool CreatePluginBuildFile(const FString& NewBuildFileName, const FString& PluginName, FText& OutFailReason, FString TemplateType = FString(""), TArray<FString> InPrivateDependencyModuleNames = TArray<FString>());
-
-	/**
-	 * Create plugin header file.
-	 * @param FolderPath - filename with full path
-	 * @param PluginName - name of the plugin this file belongs to
-	 * @param OutFailReason - fail reason text if something goes wrong
-	 * @param TemplateType - empty by default; can be Blank, Basic or Advanced
-	 * @return true if file was created, false otherwise
-	 */
-	static bool CreatePluginHeaderFile(const FString& FolderPath, const FString& PluginName, FText& OutFailReason, FString TemplateType = FString(""));
-
-	/**
-	 * Create plugin style files (source and header).
-	 * @param PrivateFolderPath - Private folder path for source file
-	 * @param PublicFolderPath - Public folder path for header file
-	 * @param PluginName - name of the plugin this file belongs to
-	 * @param OutFailReason - fail reason text if something goes wrong
-	 * @param TemplateType - empty by default; can be Blank, Basic or Advanced
-	 * @return true if file was created, false otherwise
-	 */
-	static bool CreatePluginStyleFiles(const FString& PrivateFolderPath, const FString& PublicFolderPath, const FString& PluginName, FText& OutFailReason, FString TemplateType = FString(""));
-
-	/**
-	 * Create plugin editor mode files (source and header).
-	 * @param PrivateFolderPath - Private folder path for source file
-	 * @param PublicFolderPath - Public folder path for header file
-	 * @param PluginName - name of the plugin this file belongs to
-	 * @param bUsesToolkits - whether the editor mode uses toolkits
-	 * @param bIncludeSampleUI - whether the toolkit includes sample UI if enabled
-	 * @param OutFailReason - fail reason text if something goes wrong
-	 * @param TemplateType - empty by default; can be Blank, Basic or Advanced
-	 * @return true if file was created, false otherwise
-	 */
-	static bool CreatePluginEdModeFiles(const FString& PrivateFolderPath, const FString& PublicFolderPath, const FString& PluginName, bool bUsesToolkits, bool bIncludeSampleUI, FText& OutFailReason, FString TemplateType = FString(""));
-
-	/**
-	 * Create plugin source file.
-	 * @param FolderPath - filename with full path
-	 * @param PluginName - name of the plugin this file belongs to
-	 * @param OutFailReason - fail reason text if something goes wrong
-	 * @param TemplateType - empty by default; can be Blank, Basic or Advanced
-	 * @param bMakeEditorMode - Whether we are making cpp file for plugin with Editor Mode
-	 * @param bUsesToolkits - Whether the Editor Mode is using toolkits
-	 * @return true if file was created, false otherwise
-	 */
-	static bool CreatePluginCPPFile(const FString& FolderPath, const FString& PluginName, FText& OutFailReason, FString TemplateType = FString(""), bool bMakeEditorMode = false, bool bUsesToolkits = false);
-
-	/**
-	 * Create plugin UI commands file.
-	 * @param FolderPath - filename with full path
-	 * @param PluginName - name of the plugin this file belongs to
-	 * @param OutFailReason - fail reason text if something goes wrong
-	 * @param TemplateType - empty by default; can be Blank, Basic or Advanced
-	 * @return true if file was created, false otherwise
-	 */
-	static bool CreateCommandsFiles(const FString& FolderPath, const FString& PluginName, FText& OutFailReason, FString TemplateType = FString(""));
-
-	/**
-	 * Create plugin private precompiled header file.
-	 * @param FolderPath - filename with full path
-	 * @param PluginName - name of the plugin this file belongs to
-	 * @param OutFailReason - fail reason text if something goes wrong
-	 * @param TemplateType - empty by default; can be Blank, Basic or Advanced
-	 * @return true if file was created, false otherwise
-	 */
-	static bool CreatePrivatePCHFile(const FString& FolderPath, const FString& PluginName, FText& OutFailReason, FString TemplateType = FString(""));
-
-	/**
-	 * Create Blueprint Function library class file.
-	 * @param PrivateFolderPath - Private folder path for source file
-	 * @param PublicFolderPath - Public folder path for header file
-	 * @param PluginName - name of the plugin this file belongs to
-	 * @param OutFailReason - fail reason text if something goes wrong
-	 * @return true if file was created, false otherwise
-	 */
-	static bool CreateBlueprintFunctionLibraryFiles(const FString& PrivateFolderPath, const FString& PublicFolderPath, const FString& PluginName, FText& OutFailReason);
-
-	/**
-	 * Make string that will register editor mode.
-	 * @param EdModeName The name of the editor mode
-	 * @param bUsesToolkits Whether the editor mode uses toolkits
-	 * @return - string which will be used to register editor mode in PluginModule.cpp (template) or empty string if there is no mode
-	 */
-	static FString MakeRegisterEdModeString(const FString& EdModeName, bool bUsesToolkits)
-	{
-		if (EdModeName.IsEmpty())
-		{
-			return FString("");
-		}
-
-		return FString::Printf(TEXT("FEditorModeRegistry::Get().RegisterMode<F%s>(F%s::EM_%sId, LOCTEXT(\"%sName\", \"%s\"), FSlateIcon(), %s);"), *EdModeName, *EdModeName, *EdModeName, *EdModeName, *EdModeName, bUsesToolkits ? TEXT("true") : TEXT("false"));
-	}
-
-	/**
-	 * Make string that will unregister editor mode.
-	 * @param EdModeName The name of the editor mode
-	 * @return - string which will be used to unregister editor mode in PluginModule.cpp (template) or empty string if there is no mode
-	 */
-	static FString MakeUnRegisterEdModeString(const FString& EdModeName)
-	{
-		if (EdModeName.IsEmpty())
-		{
-			return FString("");
-		}
-
-		return FString::Printf(TEXT("FEditorModeRegistry::Get().UnregisterMode(F%s::EM_%sId);"), *EdModeName, *EdModeName);
-	}
-
 	/**
 	 * Copy the contents of a template folder to create a new plugin, replacing instances of PLUGIN_NAME with the new plugin name
 	 * @param DestinationDirectory Base Directory that the template directory should be copied into
diff --git a/Engine/Plugins/Editor/PluginBrowser/Source/PluginBrowser/Private/SNewPluginWizard.cpp b/Engine/Plugins/Editor/PluginBrowser/Source/PluginBrowser/Private/SNewPluginWizard.cpp
--- a/Engine/Plugins/Editor/PluginBrowser/Source/PluginBrowser/Private/SNewPluginWizard.cpp
+++ b/Engine/Plugins/Editor/PluginBrowser/Source/PluginBrowser/Private/SNewPluginWizard.cpp
@@ -29,6 +29,7 @@ SNewPluginWizard::SNewPluginWizard()
 	const FText BlueprintLibTemplateName = LOCTEXT("BlueprintLibTemplateLabel", "Blueprint Library");
 	const FText EditorModeTemplateName = LOCTEXT("EditorModeTemplateLabel", "Editor Mode");
 	const FText ThirdPartyTemplateName = LOCTEXT("ThirdPartyTemplateLabel", "Third Party Library");
+
 	const FText BlankDescription = LOCTEXT("BlankTemplateDesc", "Create a blank plugin with a minimal amount of code.\n\nChoose this if you want to set everything up from scratch or are making a non-visual plugin.\nA plugin created with this template will appear in the Editor's plugin list but will not register any buttons or menu entries.");
 	const FText BasicDescription = LOCTEXT("BasicTemplateDesc", "Create a plugin that will add a button to the toolbar in the Level Editor.\n\nStart by implementing something in the created \"OnButtonClick\" event.");
 	const FText AdvancedDescription = LOCTEXT("AdvancedTemplateDesc", "Create a plugin that will add a button to the toolbar in the Level Editor that summons an empty standalone tab window when clicked.");
@@ -36,12 +37,12 @@ SNewPluginWizard::SNewPluginWizard()
 	const FText EditorModeDescription = LOCTEXT("EditorModeDesc", "Create a plugin that will have an editor mode.\n\nThis will include a toolkit example to specify UI that will appear in \"Modes\" tab (next to Foliage, Landscape etc).\nIt will also include very basic UI that demonstrates editor interaction and undo/redo functions usage.");
 	const FText ThirdPartyDescription = LOCTEXT("ThirdPartyDesc", "Create a plugin that uses an included third party library.\n\nThis can be used as an example of how to include, load and use a third party library yourself.");
 
-	Templates.Add(MakeShareable(new FPluginTemplateDescription(BlankTemplateName, BlankDescription, TEXT("Blank"), false, false, false, false)));
-	Templates.Add(MakeShareable(new FPluginTemplateDescription(BasicTemplateName, BasicDescription, TEXT("Basic"), false, true, false, false)));
-	Templates.Add(MakeShareable(new FPluginTemplateDescription(AdvancedTemplateName, AdvancedDescription, TEXT("Advanced"), false, true, false, false)));
-	Templates.Add(MakeShareable(new FPluginTemplateDescription(BlueprintLibTemplateName, BlueprintLibDescription, TEXT("Blank"), false, false, false, true)));
-	Templates.Add(MakeShareable(new FPluginTemplateDescription(EditorModeTemplateName, EditorModeDescription, TEXT("Blank"), false, false, true, false)));
-	Templates.Add(MakeShareable(new FPluginTemplateDescription(ThirdPartyTemplateName, ThirdPartyDescription, TEXT("ThirdPartyLibrary"), true, false, false, false)));
+	Templates.Add(MakeShareable(new FPluginTemplateDescription(BlankTemplateName, BlankDescription, TEXT("Blank"))));
+	Templates.Add(MakeShareable(new FPluginTemplateDescription(BasicTemplateName, BasicDescription, TEXT("Basic"))));
+	Templates.Add(MakeShareable(new FPluginTemplateDescription(AdvancedTemplateName, AdvancedDescription, TEXT("Advanced"))));
+	Templates.Add(MakeShareable(new FPluginTemplateDescription(BlueprintLibTemplateName, BlueprintLibDescription, TEXT("BlueprintLibrary"))));
+	Templates.Add(MakeShareable(new FPluginTemplateDescription(EditorModeTemplateName, EditorModeDescription, TEXT("EditorMode"))));
+	Templates.Add(MakeShareable(new FPluginTemplateDescription(ThirdPartyTemplateName, ThirdPartyDescription, TEXT("ThirdPartyLibrary"))));
 
 	AbsoluteGamePluginPath = IFileManager::Get().ConvertToAbsolutePathForExternalAppForWrite(*FPaths::GamePluginsDir());
 	FPaths::MakePlatformFilename(AbsoluteGamePluginPath);
@@ -434,17 +435,15 @@ FReply SNewPluginWizard::OnCreatePluginClicked()
 	const FString AutoPluginName = PluginNameText.ToString();
 	const FPluginTemplateDescription& SelectedTemplate = *CurrentTemplate.Get();
 
-	// For the time being, add all the features to editor mode templates
-	bool bUsesToolkits = SelectedTemplate.bMakeEditorMode;
-	bool bIncludeSampleUI = SelectedTemplate.bMakeEditorMode && bUsesToolkits;
-
 	// Plugin thumbnail image
+	bool bRequiresDefaultIcon = false;
 	FString PluginBaseDir = IPluginManager::Get().FindPlugin(TEXT("PluginBrowser"))->GetBaseDir();
 	FString TemplateFolderName = PluginBaseDir / TEXT("Templates") / SelectedTemplate.OnDiskPath;
 	FString PluginEditorIconPath = TemplateFolderName / TEXT("Resources/Icon128.png");
 	if (!FPlatformFileManager::Get().GetPlatformFile().FileExists(*PluginEditorIconPath))
 	{
 		PluginEditorIconPath = PluginBaseDir / TEXT("Resources/DefaultIcon128.png");
+		bRequiresDefaultIcon = true;
 	}
 
 	TArray<FString> CreatedFiles;
@@ -463,107 +462,16 @@ FReply SNewPluginWizard::OnCreatePluginClicked()
 	// Resource folder
 	const FString ResourcesFolder = PluginFolder / TEXT("Resources");
 
-	// Copy the icons
-	bSucceeded = bSucceeded && CopyFile(ResourcesFolder / TEXT("Icon128.png"), PluginEditorIconPath, /*inout*/ CreatedFiles);
-
-	if (SelectedTemplate.bCopyWholeDirectory)
+	if (bRequiresDefaultIcon)
 	{
-		if (bSucceeded && !FPluginHelpers::CopyPluginTemplateFolder(*PluginFolder, *TemplateFolderName, AutoPluginName))
-		{
-			PopErrorNotification(FText::Format(LOCTEXT("FailedTemplateCopy", "Failed to copy plugin Template: {0}"), FText::FromString(TemplateFolderName)));
-			bSucceeded = false;
-		}
+		// Copy the icon
+		bSucceeded = bSucceeded && CopyFile(ResourcesFolder / TEXT("Icon128.png"), PluginEditorIconPath, /*inout*/ CreatedFiles);
 	}
-	else
-	{
-
-		if (SelectedTemplate.bIncludeUI)
-		{
-			FString PluginButtonIconPath = PluginBaseDir / TEXT("Resources/ButtonIcon_40x.png");
-			bSucceeded = bSucceeded && CopyFile(ResourcesFolder / TEXT("ButtonIcon_40x.png"), PluginButtonIconPath, /*inout*/ CreatedFiles);
-		}
-
-		// Source folder
-		const FString SourceFolder = PluginFolder / TEXT("Source");
-
-		// Create Plugin Module Folders
-		const FString PluginModuleFolder = SourceFolder / AutoPluginName;
-		FString PrivateSourceFolder = PluginModuleFolder / TEXT("Private");
-		FString PublicSourceFolder = PluginModuleFolder / TEXT("Public");
-
-		TArray<FString> PrivateDependencyModuleNames;
-		PrivateDependencyModuleNames.Add("CoreUObject");
-		PrivateDependencyModuleNames.Add("Engine");
-		PrivateDependencyModuleNames.Add("Slate");
-		PrivateDependencyModuleNames.Add("SlateCore");
-
-		// If we're going to create editor mode, make sure that all build files will have necessary modules dependencies
-		// Only Blank plugin needs additional modules dependencies
-		// @TODO: find better way check if selected template is Blank template
-		if (SelectedTemplate.bMakeEditorMode && !SelectedTemplate.bIncludeUI)
-		{
-			PrivateDependencyModuleNames.AddUnique(TEXT("InputCore"));
-			PrivateDependencyModuleNames.AddUnique(TEXT("UnrealEd"));
-			PrivateDependencyModuleNames.AddUnique(TEXT("LevelEditor"));
-		}
-
-		// Based on chosen template create build, and other source files
-		if (bSucceeded && !FPluginHelpers::CreatePluginBuildFile(PluginModuleFolder / AutoPluginName + TEXT(".Build.cs"), AutoPluginName, LocalFailReason, SelectedTemplate.OnDiskPath, PrivateDependencyModuleNames))
-		{
-			PopErrorNotification(FText::Format(LOCTEXT("FailedBuild", "Failed to create plugin build file. {0}"), LocalFailReason));
-			bSucceeded = false;
-		}
-
-		if (bSucceeded && !FPluginHelpers::CreatePluginHeaderFile(PublicSourceFolder, AutoPluginName, LocalFailReason, SelectedTemplate.OnDiskPath))
-		{
-			PopErrorNotification(FText::Format(LOCTEXT("FailedHeader", "Failed to create plugin header file. {0}"), LocalFailReason));
-			bSucceeded = false;
-		}
-
-		if (bSucceeded && !FPluginHelpers::CreatePrivatePCHFile(PrivateSourceFolder, AutoPluginName, LocalFailReason, SelectedTemplate.OnDiskPath))
-		{
-			PopErrorNotification(FText::Format(LOCTEXT("FailedPCH", "Failed to create plugin PCH file. {0}"), LocalFailReason));
-			bSucceeded = false;
-		}
-
-		if (bSucceeded && !FPluginHelpers::CreatePluginCPPFile(PrivateSourceFolder, AutoPluginName, LocalFailReason, SelectedTemplate.OnDiskPath, SelectedTemplate.bMakeEditorMode, bUsesToolkits))
-		{
-			PopErrorNotification(FText::Format(LOCTEXT("FailedCppFile", "Failed to create plugin cpp file. {0}"), LocalFailReason));
-			bSucceeded = false;
-		}
-
-		if (SelectedTemplate.bIncludeUI)
-		{
-			if (bSucceeded && !FPluginHelpers::CreatePluginStyleFiles(PrivateSourceFolder, PublicSourceFolder, AutoPluginName, LocalFailReason, SelectedTemplate.OnDiskPath))
-			{
-				PopErrorNotification(FText::Format(LOCTEXT("FailedStylesFile", "Failed to create plugin styles files. {0}"), LocalFailReason));
-				bSucceeded = false;
-			}
 
-			if (bSucceeded && !FPluginHelpers::CreateCommandsFiles(PublicSourceFolder, AutoPluginName, LocalFailReason, SelectedTemplate.OnDiskPath))
-			{
-				PopErrorNotification(FText::Format(LOCTEXT("FailedStylesFile", "Failed to create plugin commands files. {0}"), LocalFailReason));
-				bSucceeded = false;
-			}
-		}
-
-		if (SelectedTemplate.bAddBPLibrary)
-		{
-			if (bSucceeded && !FPluginHelpers::CreateBlueprintFunctionLibraryFiles(PrivateSourceFolder, PublicSourceFolder, AutoPluginName, LocalFailReason))
-			{
-				PopErrorNotification(FText::Format(LOCTEXT("FailedBPLibraryFile", "Failed to create blueprint library files. {0}"), LocalFailReason));
-				bSucceeded = false;
-			}
-		}
-
-		if (SelectedTemplate.bMakeEditorMode)
-		{
-			if (bSucceeded && !FPluginHelpers::CreatePluginEdModeFiles(PrivateSourceFolder, PublicSourceFolder, AutoPluginName, bUsesToolkits, bIncludeSampleUI, LocalFailReason))
-			{
-				PopErrorNotification(FText::Format(LOCTEXT("FailedBPLibraryFile", "Failed to create EdMode files. {0}"), LocalFailReason));
-				bSucceeded = false;
-			}
-		}
+	if (bSucceeded && !FPluginHelpers::CopyPluginTemplateFolder(*PluginFolder, *TemplateFolderName, AutoPluginName))
+	{
+		PopErrorNotification(FText::Format(LOCTEXT("FailedTemplateCopy", "Failed to copy plugin Template: {0}"), FText::FromString(TemplateFolderName)));
+		bSucceeded = false;
 	}
 
 	if (bSucceeded)
diff --git a/Engine/Plugins/Editor/PluginBrowser/Source/PluginBrowser/Private/SNewPluginWizard.h b/Engine/Plugins/Editor/PluginBrowser/Source/PluginBrowser/Private/SNewPluginWizard.h
--- a/Engine/Plugins/Editor/PluginBrowser/Source/PluginBrowser/Private/SNewPluginWizard.h
+++ b/Engine/Plugins/Editor/PluginBrowser/Source/PluginBrowser/Private/SNewPluginWizard.h
@@ -20,30 +20,14 @@ struct FPluginTemplateDescription
 	/** Name of the directory containing template files */
 	FString OnDiskPath;
 
-	/** Whether to copy the whole template directory rather than create new files from .template files */
-	bool bCopyWholeDirectory;
-
-	/** Whether this template needs the UI files generated (commands, etc...) */
-	bool bIncludeUI;
-
-	/** Whether this template creates an editor mode */
-	bool bMakeEditorMode;
-
-	/** Whether this template creates a blueprint function library */
-	bool bAddBPLibrary;
-
 	/** Brush resource for the image that is dynamically loaded */
 	TSharedPtr< FSlateDynamicImageBrush > PluginIconDynamicImageBrush;
 
 	/** Constructor */
-	FPluginTemplateDescription(FText InName, FText InDescription, FString InOnDiskPath, bool bInCopyWholeDirectory, bool bInIncludeUI, bool bInMakeEditorMode, bool bInAddBPLibrary)
+	FPluginTemplateDescription(FText InName, FText InDescription, FString InOnDiskPath)
 		: Name(InName)
 		, Description(InDescription)
 		, OnDiskPath(InOnDiskPath)
-		, bCopyWholeDirectory(bInCopyWholeDirectory)
-		, bIncludeUI(bInIncludeUI)
-		, bMakeEditorMode(bInMakeEditorMode)
-		, bAddBPLibrary(bInAddBPLibrary)
 	{
 	}
 };
diff --git a/Engine/Plugins/Editor/PluginBrowser/Templates/Advanced/PluginCommands.cpp.template b/Engine/Plugins/Editor/PluginBrowser/Templates/Advanced/PluginCommands.cpp.template
deleted file mode 100644
diff --git a/Engine/Plugins/Editor/PluginBrowser/Templates/Advanced/PluginCommands.h.template b/Engine/Plugins/Editor/PluginBrowser/Templates/Advanced/PluginCommands.h.template
deleted file mode 100644
diff --git a/Engine/Plugins/Editor/PluginBrowser/Templates/Advanced/PluginModule.Build.cs.template b/Engine/Plugins/Editor/PluginBrowser/Templates/Advanced/PluginModule.Build.cs.template
deleted file mode 100644
diff --git a/Engine/Plugins/Editor/PluginBrowser/Templates/Advanced/PluginModule.cpp.template b/Engine/Plugins/Editor/PluginBrowser/Templates/Advanced/PluginModule.cpp.template
deleted file mode 100644
diff --git a/Engine/Plugins/Editor/PluginBrowser/Templates/Advanced/PluginModule.h.template b/Engine/Plugins/Editor/PluginBrowser/Templates/Advanced/PluginModule.h.template
deleted file mode 100644
diff --git a/Engine/Plugins/Editor/PluginBrowser/Templates/Advanced/PluginModulePCH.h.template b/Engine/Plugins/Editor/PluginBrowser/Templates/Advanced/PluginModulePCH.h.template
deleted file mode 100644
diff --git a/Engine/Plugins/Editor/PluginBrowser/Templates/Advanced/PluginStyle.cpp.template b/Engine/Plugins/Editor/PluginBrowser/Templates/Advanced/PluginStyle.cpp.template
deleted file mode 100644
diff --git a/Engine/Plugins/Editor/PluginBrowser/Templates/Advanced/PluginStyle.h.template b/Engine/Plugins/Editor/PluginBrowser/Templates/Advanced/PluginStyle.h.template
deleted file mode 100644
diff --git a/Engine/Plugins/Editor/PluginBrowser/Templates/Advanced/Source/PLUGIN_NAME/PLUGIN_NAME.Build.cs b/Engine/Plugins/Editor/PluginBrowser/Templates/Advanced/Source/PLUGIN_NAME/PLUGIN_NAME.Build.cs
new file mode 100644
--- /dev/null
+++ b/Engine/Plugins/Editor/PluginBrowser/Templates/Advanced/Source/PLUGIN_NAME/PLUGIN_NAME.Build.cs
@@ -0,0 +1,58 @@
+// Copyright 1998-2015 Epic Games, Inc. All Rights Reserved.
+
+using UnrealBuildTool;
+
+public class PLUGIN_NAME : ModuleRules
+{
+	public PLUGIN_NAME(TargetInfo Target)
+	{
+		
+		PublicIncludePaths.AddRange(
+			new string[] {
+				"PLUGIN_NAME/Public"
+				// ... add public include paths required here ...
+			}
+			);
+				
+		
+		PrivateIncludePaths.AddRange(
+			new string[] {
+				"PLUGIN_NAME/Private",
+				// ... add other private include paths required here ...
+			}
+			);
+			
+		
+		PublicDependencyModuleNames.AddRange(
+			new string[]
+			{
+				"Core",
+				// ... add other public dependencies that you statically link with here ...
+			}
+			);
+			
+		
+		PrivateDependencyModuleNames.AddRange(
+			new string[]
+			{
+				"Projects",
+				"InputCore",
+				"UnrealEd",
+				"LevelEditor",
+				"CoreUObject",
+				"Engine",
+				"Slate",
+				"SlateCore",
+				// ... add private dependencies that you statically link with here ...	
+			}
+			);
+		
+		
+		DynamicallyLoadedModuleNames.AddRange(
+			new string[]
+			{
+				// ... add any modules that your module loads dynamically here ...
+			}
+			);
+	}
+}
diff --git a/Engine/Plugins/Editor/PluginBrowser/Templates/Advanced/Source/PLUGIN_NAME/Private/PLUGIN_NAME.cpp b/Engine/Plugins/Editor/PluginBrowser/Templates/Advanced/Source/PLUGIN_NAME/Private/PLUGIN_NAME.cpp
new file mode 100644
--- /dev/null
+++ b/Engine/Plugins/Editor/PluginBrowser/Templates/Advanced/Source/PLUGIN_NAME/Private/PLUGIN_NAME.cpp
@@ -0,0 +1,104 @@
+// Copyright 1998-2015 Epic Games, Inc. All Rights Reserved.
+
+#include "PLUGIN_NAMEPrivatePCH.h"
+
+#include "SlateBasics.h"
+#include "SlateExtras.h"
+
+#include "PLUGIN_NAMEStyle.h"
+#include "PLUGIN_NAMECommands.h"
+
+#include "LevelEditor.h"
+
+static const FName PLUGIN_NAMETabName("PLUGIN_NAME");
+
+#define LOCTEXT_NAMESPACE "FPLUGIN_NAMEModule"
+
+void FPLUGIN_NAMEModule::StartupModule()
+{
+	// This code will execute after your module is loaded into memory; the exact timing is specified in the .uplugin file per-module
+	
+	FPLUGIN_NAMEStyle::Initialize();
+	FPLUGIN_NAMEStyle::ReloadTextures();
+
+	FPLUGIN_NAMECommands::Register();
+	
+	PluginCommands = MakeShareable(new FUICommandList);
+
+	PluginCommands->MapAction(
+		FPLUGIN_NAMECommands::Get().OpenPluginWindow,
+		FExecuteAction::CreateRaw(this, &FPLUGIN_NAMEModule::PluginButtonClicked),
+		FCanExecuteAction());
+		
+	FLevelEditorModule& LevelEditorModule = FModuleManager::LoadModuleChecked<FLevelEditorModule>("LevelEditor");
+	
+	{
+		TSharedPtr<FExtender> MenuExtender = MakeShareable(new FExtender());
+		MenuExtender->AddMenuExtension("WindowLayout", EExtensionHook::After, PluginCommands, FMenuExtensionDelegate::CreateRaw(this, &FPLUGIN_NAMEModule::AddMenuExtension));
+
+		LevelEditorModule.GetMenuExtensibilityManager()->AddExtender(MenuExtender);
+	}
+	
+	{
+		TSharedPtr<FExtender> ToolbarExtender = MakeShareable(new FExtender);
+		ToolbarExtender->AddToolBarExtension("Settings", EExtensionHook::After, PluginCommands, FToolBarExtensionDelegate::CreateRaw(this, &FPLUGIN_NAMEModule::AddToolbarExtension));
+		
+		LevelEditorModule.GetToolBarExtensibilityManager()->AddExtender(ToolbarExtender);
+	}
+	
+	FGlobalTabmanager::Get()->RegisterNomadTabSpawner(PLUGIN_NAMETabName, FOnSpawnTab::CreateRaw(this, &FPLUGIN_NAMEModule::OnSpawnPluginTab))
+		.SetDisplayName(LOCTEXT("FPLUGIN_NAMETabTitle", "PLUGIN_NAME"))
+		.SetMenuType(ETabSpawnerMenuType::Hidden);
+}
+
+void FPLUGIN_NAMEModule::ShutdownModule()
+{
+	// This function may be called during shutdown to clean up your module.  For modules that support dynamic reloading,
+	// we call this function before unloading the module.
+	FPLUGIN_NAMEStyle::Shutdown();
+
+	FPLUGIN_NAMECommands::Unregister();
+
+	FGlobalTabmanager::Get()->UnregisterNomadTabSpawner(PLUGIN_NAMETabName);
+}
+
+TSharedRef<SDockTab> FPLUGIN_NAMEModule::OnSpawnPluginTab(const FSpawnTabArgs& SpawnTabArgs)
+{
+	FText WidgetText = FText::Format(
+		LOCTEXT("WindowWidgetText", "Add code to {0} in {1} to override this window's contents"),
+		FText::FromString(TEXT("FPLUGIN_NAMEModule::OnSpawnPluginTab")),
+		FText::FromString(TEXT("PLUGIN_NAME.cpp"))
+		);
+
+	return SNew(SDockTab)
+		.TabRole(ETabRole::NomadTab)
+		[
+			// Put your tab content here!
+			SNew(SBox)
+			.HAlign(HAlign_Center)
+			.VAlign(VAlign_Center)
+			[
+				SNew(STextBlock)
+				.Text(WidgetText)
+			]
+		];
+}
+
+void FPLUGIN_NAMEModule::PluginButtonClicked()
+{
+	FGlobalTabmanager::Get()->InvokeTab(PLUGIN_NAMETabName);
+}
+
+void FPLUGIN_NAMEModule::AddMenuExtension(FMenuBuilder& Builder)
+{
+	Builder.AddMenuEntry(FPLUGIN_NAMECommands::Get().OpenPluginWindow);
+}
+
+void FPLUGIN_NAMEModule::AddToolbarExtension(FToolBarBuilder& Builder)
+{
+	Builder.AddToolBarButton(FPLUGIN_NAMECommands::Get().OpenPluginWindow);
+}
+
+#undef LOCTEXT_NAMESPACE
+	
+IMPLEMENT_MODULE(FPLUGIN_NAMEModule, PLUGIN_NAME)
\ No newline at end of file
diff --git a/Engine/Plugins/Editor/PluginBrowser/Templates/Advanced/Source/PLUGIN_NAME/Private/PLUGIN_NAMEPrivatePCH.h b/Engine/Plugins/Editor/PluginBrowser/Templates/Advanced/Source/PLUGIN_NAME/Private/PLUGIN_NAMEPrivatePCH.h
new file mode 100644
--- /dev/null
+++ b/Engine/Plugins/Editor/PluginBrowser/Templates/Advanced/Source/PLUGIN_NAME/Private/PLUGIN_NAMEPrivatePCH.h
@@ -0,0 +1,8 @@
+// Copyright 1998-2015 Epic Games, Inc. All Rights Reserved.
+
+#include "SlateBasics.h"
+
+#include "PLUGIN_NAME.h"
+
+// You should place include statements to your module's private header files here.  You only need to
+// add includes for headers that are used in most of your module's source files though.
\ No newline at end of file
diff --git a/Engine/Plugins/Editor/PluginBrowser/Templates/Advanced/Source/PLUGIN_NAME/Private/PLUGIN_NAMEStyle.cpp b/Engine/Plugins/Editor/PluginBrowser/Templates/Advanced/Source/PLUGIN_NAME/Private/PLUGIN_NAMEStyle.cpp
new file mode 100644
--- /dev/null
+++ b/Engine/Plugins/Editor/PluginBrowser/Templates/Advanced/Source/PLUGIN_NAME/Private/PLUGIN_NAMEStyle.cpp
@@ -0,0 +1,66 @@
+// Copyright 1998-2015 Epic Games, Inc. All Rights Reserved.
+
+#include "PLUGIN_NAMEPrivatePCH.h"
+#include "PLUGIN_NAMEStyle.h"
+#include "SlateGameResources.h"
+#include "IPluginManager.h"
+
+TSharedPtr< FSlateStyleSet > FPLUGIN_NAMEStyle::StyleInstance = NULL;
+
+void FPLUGIN_NAMEStyle::Initialize()
+{
+	if (!StyleInstance.IsValid())
+	{
+		StyleInstance = Create();
+		FSlateStyleRegistry::RegisterSlateStyle(*StyleInstance);
+	}
+}
+
+void FPLUGIN_NAMEStyle::Shutdown()
+{
+	FSlateStyleRegistry::UnRegisterSlateStyle(*StyleInstance);
+	ensure(StyleInstance.IsUnique());
+	StyleInstance.Reset();
+}
+
+FName FPLUGIN_NAMEStyle::GetStyleSetName()
+{
+	static FName StyleSetName(TEXT("PLUGIN_NAMEStyle"));
+	return StyleSetName;
+}
+
+#define IMAGE_BRUSH( RelativePath, ... ) FSlateImageBrush( Style->RootToContentDir( RelativePath, TEXT(".png") ), __VA_ARGS__ )
+#define BOX_BRUSH( RelativePath, ... ) FSlateBoxBrush( Style->RootToContentDir( RelativePath, TEXT(".png") ), __VA_ARGS__ )
+#define BORDER_BRUSH( RelativePath, ... ) FSlateBorderBrush( Style->RootToContentDir( RelativePath, TEXT(".png") ), __VA_ARGS__ )
+#define TTF_FONT( RelativePath, ... ) FSlateFontInfo( Style->RootToContentDir( RelativePath, TEXT(".ttf") ), __VA_ARGS__ )
+#define OTF_FONT( RelativePath, ... ) FSlateFontInfo( Style->RootToContentDir( RelativePath, TEXT(".otf") ), __VA_ARGS__ )
+
+const FVector2D Icon16x16(16.0f, 16.0f);
+const FVector2D Icon20x20(20.0f, 20.0f);
+const FVector2D Icon40x40(40.0f, 40.0f);
+
+TSharedRef< FSlateStyleSet > FPLUGIN_NAMEStyle::Create()
+{
+	TSharedRef< FSlateStyleSet > Style = MakeShareable(new FSlateStyleSet("PLUGIN_NAMEStyle"));
+	Style->SetContentRoot(IPluginManager::Get().FindPlugin("PLUGIN_NAME")->GetBaseDir() / TEXT("Resources"));
+
+	Style->Set("PLUGIN_NAME.OpenPluginWindow", new IMAGE_BRUSH(TEXT("ButtonIcon_40x"), Icon40x40));
+
+	return Style;
+}
+
+#undef IMAGE_BRUSH
+#undef BOX_BRUSH
+#undef BORDER_BRUSH
+#undef TTF_FONT
+#undef OTF_FONT
+
+void FPLUGIN_NAMEStyle::ReloadTextures()
+{
+	FSlateApplication::Get().GetRenderer()->ReloadTextureResources();
+}
+
+const ISlateStyle& FPLUGIN_NAMEStyle::Get()
+{
+	return *StyleInstance;
+}
diff --git a/Engine/Plugins/Editor/PluginBrowser/Templates/Advanced/Source/PLUGIN_NAME/Public/PLUGIN_NAME.h b/Engine/Plugins/Editor/PluginBrowser/Templates/Advanced/Source/PLUGIN_NAME/Public/PLUGIN_NAME.h
new file mode 100644
--- /dev/null
+++ b/Engine/Plugins/Editor/PluginBrowser/Templates/Advanced/Source/PLUGIN_NAME/Public/PLUGIN_NAME.h
@@ -0,0 +1,30 @@
+// Copyright 1998-2015 Epic Games, Inc. All Rights Reserved.
+
+#pragma once
+
+#include "ModuleManager.h"
+
+class FToolBarBuilder;
+class FMenuBuilder;
+
+class FPLUGIN_NAMEModule : public IModuleInterface
+{
+public:
+
+	/** IModuleInterface implementation */
+	virtual void StartupModule() override;
+	virtual void ShutdownModule() override;
+	
+	/** This function will be bound to Command (by default it will bring up plugin window) */
+	void PluginButtonClicked();
+	
+private:
+
+	void AddToolbarExtension(FToolBarBuilder& Builder);
+	void AddMenuExtension(FMenuBuilder& Builder);
+
+	TSharedRef<class SDockTab> OnSpawnPluginTab(const class FSpawnTabArgs& SpawnTabArgs);
+
+private:
+	TSharedPtr<class FUICommandList> PluginCommands;
+};
\ No newline at end of file
diff --git a/Engine/Plugins/Editor/PluginBrowser/Templates/Advanced/Source/PLUGIN_NAME/Public/PLUGIN_NAMECommands.cpp b/Engine/Plugins/Editor/PluginBrowser/Templates/Advanced/Source/PLUGIN_NAME/Public/PLUGIN_NAMECommands.cpp
new file mode 100644
--- /dev/null
+++ b/Engine/Plugins/Editor/PluginBrowser/Templates/Advanced/Source/PLUGIN_NAME/Public/PLUGIN_NAMECommands.cpp
@@ -0,0 +1,13 @@
+// Copyright 1998-2015 Epic Games, Inc. All Rights Reserved.
+
+#include "PLUGIN_NAMEPrivatePCH.h"
+#include "PLUGIN_NAMECommands.h"
+
+#define LOCTEXT_NAMESPACE "FPLUGIN_NAMEModule"
+
+void FPLUGIN_NAMECommands::RegisterCommands()
+{
+	UI_COMMAND(OpenPluginWindow, "PLUGIN_NAME", "Bring up PLUGIN_NAME window", EUserInterfaceActionType::Button, FInputGesture());
+}
+
+#undef LOCTEXT_NAMESPACE
diff --git a/Engine/Plugins/Editor/PluginBrowser/Templates/Advanced/Source/PLUGIN_NAME/Public/PLUGIN_NAMECommands.h b/Engine/Plugins/Editor/PluginBrowser/Templates/Advanced/Source/PLUGIN_NAME/Public/PLUGIN_NAMECommands.h
new file mode 100644
--- /dev/null
+++ b/Engine/Plugins/Editor/PluginBrowser/Templates/Advanced/Source/PLUGIN_NAME/Public/PLUGIN_NAMECommands.h
@@ -0,0 +1,22 @@
+// Copyright 1998-2015 Epic Games, Inc. All Rights Reserved.
+
+#pragma once
+
+#include "SlateBasics.h"
+#include "PLUGIN_NAMEStyle.h"
+
+class FPLUGIN_NAMECommands : public TCommands<FPLUGIN_NAMECommands>
+{
+public:
+
+	FPLUGIN_NAMECommands()
+		: TCommands<FPLUGIN_NAMECommands>(TEXT("PLUGIN_NAME"), NSLOCTEXT("Contexts", "PLUGIN_NAME", "PLUGIN_NAME Plugin"), NAME_None, FPLUGIN_NAMEStyle::GetStyleSetName())
+	{
+	}
+
+	// TCommands<> interface
+	virtual void RegisterCommands() override;
+
+public:
+	TSharedPtr< FUICommandInfo > OpenPluginWindow;
+};
\ No newline at end of file
diff --git a/Engine/Plugins/Editor/PluginBrowser/Templates/Advanced/Source/PLUGIN_NAME/Public/PLUGIN_NAMEStyle.h b/Engine/Plugins/Editor/PluginBrowser/Templates/Advanced/Source/PLUGIN_NAME/Public/PLUGIN_NAMEStyle.h
new file mode 100644
--- /dev/null
+++ b/Engine/Plugins/Editor/PluginBrowser/Templates/Advanced/Source/PLUGIN_NAME/Public/PLUGIN_NAMEStyle.h
@@ -0,0 +1,31 @@
+// Copyright 1998-2015 Epic Games, Inc. All Rights Reserved.
+
+#pragma once
+
+#include "SlateBasics.h"
+
+/**  */
+class FPLUGIN_NAMEStyle
+{
+public:
+
+	static void Initialize();
+
+	static void Shutdown();
+
+	/** reloads textures used by slate renderer */
+	static void ReloadTextures();
+
+	/** @return The Slate style set for the Shooter game */
+	static const ISlateStyle& Get();
+
+	static FName GetStyleSetName();
+
+private:
+
+	static TSharedRef< class FSlateStyleSet > Create();
+
+private:
+
+	static TSharedPtr< class FSlateStyleSet > StyleInstance;
+};
\ No newline at end of file
diff --git a/Engine/Plugins/Editor/PluginBrowser/Templates/Basic/PluginCommands.cpp.template b/Engine/Plugins/Editor/PluginBrowser/Templates/Basic/PluginCommands.cpp.template
deleted file mode 100644
diff --git a/Engine/Plugins/Editor/PluginBrowser/Templates/Basic/PluginCommands.h.template b/Engine/Plugins/Editor/PluginBrowser/Templates/Basic/PluginCommands.h.template
deleted file mode 100644
diff --git a/Engine/Plugins/Editor/PluginBrowser/Templates/Basic/PluginModule.Build.cs.template b/Engine/Plugins/Editor/PluginBrowser/Templates/Basic/PluginModule.Build.cs.template
deleted file mode 100644
diff --git a/Engine/Plugins/Editor/PluginBrowser/Templates/Basic/PluginModule.cpp.template b/Engine/Plugins/Editor/PluginBrowser/Templates/Basic/PluginModule.cpp.template
deleted file mode 100644
diff --git a/Engine/Plugins/Editor/PluginBrowser/Templates/Basic/PluginModule.h.template b/Engine/Plugins/Editor/PluginBrowser/Templates/Basic/PluginModule.h.template
deleted file mode 100644
diff --git a/Engine/Plugins/Editor/PluginBrowser/Templates/Basic/PluginModulePCH.h.template b/Engine/Plugins/Editor/PluginBrowser/Templates/Basic/PluginModulePCH.h.template
deleted file mode 100644
diff --git a/Engine/Plugins/Editor/PluginBrowser/Templates/Basic/PluginStyle.cpp.template b/Engine/Plugins/Editor/PluginBrowser/Templates/Basic/PluginStyle.cpp.template
deleted file mode 100644
diff --git a/Engine/Plugins/Editor/PluginBrowser/Templates/Basic/PluginStyle.h.template b/Engine/Plugins/Editor/PluginBrowser/Templates/Basic/PluginStyle.h.template
deleted file mode 100644
diff --git a/Engine/Plugins/Editor/PluginBrowser/Templates/Basic/Source/PLUGIN_NAME/PLUGIN_NAME.Build.cs b/Engine/Plugins/Editor/PluginBrowser/Templates/Basic/Source/PLUGIN_NAME/PLUGIN_NAME.Build.cs
new file mode 100644
--- /dev/null
+++ b/Engine/Plugins/Editor/PluginBrowser/Templates/Basic/Source/PLUGIN_NAME/PLUGIN_NAME.Build.cs
@@ -0,0 +1,58 @@
+// Copyright 1998-2015 Epic Games, Inc. All Rights Reserved.
+
+using UnrealBuildTool;
+
+public class PLUGIN_NAME : ModuleRules
+{
+	public PLUGIN_NAME(TargetInfo Target)
+	{
+		
+		PublicIncludePaths.AddRange(
+			new string[] {
+				"PLUGIN_NAME/Public"
+				// ... add public include paths required here ...
+			}
+			);
+				
+		
+		PrivateIncludePaths.AddRange(
+			new string[] {
+				"PLUGIN_NAME/Private",
+				// ... add other private include paths required here ...
+			}
+			);
+			
+		
+		PublicDependencyModuleNames.AddRange(
+			new string[]
+			{
+				"Core",
+				// ... add other public dependencies that you statically link with here ...
+			}
+			);
+			
+		
+		PrivateDependencyModuleNames.AddRange(
+			new string[]
+			{
+				"Projects",
+				"InputCore",
+				"UnrealEd",
+				"LevelEditor",
+				"CoreUObject",
+				"Engine",
+				"Slate",
+				"SlateCore",
+				// ... add private dependencies that you statically link with here ...	
+			}
+			);
+		
+		
+		DynamicallyLoadedModuleNames.AddRange(
+			new string[]
+			{
+				// ... add any modules that your module loads dynamically here ...
+			}
+			);
+	}
+}
diff --git a/Engine/Plugins/Editor/PluginBrowser/Templates/Basic/Source/PLUGIN_NAME/Private/PLUGIN_NAME.cpp b/Engine/Plugins/Editor/PluginBrowser/Templates/Basic/Source/PLUGIN_NAME/Private/PLUGIN_NAME.cpp
new file mode 100644
--- /dev/null
+++ b/Engine/Plugins/Editor/PluginBrowser/Templates/Basic/Source/PLUGIN_NAME/Private/PLUGIN_NAME.cpp
@@ -0,0 +1,82 @@
+// Copyright 1998-2015 Epic Games, Inc. All Rights Reserved.
+
+#include "PLUGIN_NAMEPrivatePCH.h"
+
+#include "SlateBasics.h"
+#include "SlateExtras.h"
+
+#include "PLUGIN_NAMEStyle.h"
+#include "PLUGIN_NAMECommands.h"
+
+#include "LevelEditor.h"
+
+static const FName PLUGIN_NAMETabName("PLUGIN_NAME");
+
+#define LOCTEXT_NAMESPACE "FPLUGIN_NAMEModule"
+
+void FPLUGIN_NAMEModule::StartupModule()
+{
+	// This code will execute after your module is loaded into memory; the exact timing is specified in the .uplugin file per-module
+	
+	FPLUGIN_NAMEStyle::Initialize();
+	FPLUGIN_NAMEStyle::ReloadTextures();
+
+	FPLUGIN_NAMECommands::Register();
+	
+	PluginCommands = MakeShareable(new FUICommandList);
+
+	PluginCommands->MapAction(
+		FPLUGIN_NAMECommands::Get().PluginAction,
+		FExecuteAction::CreateRaw(this, &FPLUGIN_NAMEModule::PluginButtonClicked),
+		FCanExecuteAction());
+		
+	FLevelEditorModule& LevelEditorModule = FModuleManager::LoadModuleChecked<FLevelEditorModule>("LevelEditor");
+	
+	{
+		TSharedPtr<FExtender> MenuExtender = MakeShareable(new FExtender());
+		MenuExtender->AddMenuExtension("WindowLayout", EExtensionHook::After, PluginCommands, FMenuExtensionDelegate::CreateRaw(this, &FPLUGIN_NAMEModule::AddMenuExtension));
+
+		LevelEditorModule.GetMenuExtensibilityManager()->AddExtender(MenuExtender);
+	}
+	
+	{
+		TSharedPtr<FExtender> ToolbarExtender = MakeShareable(new FExtender);
+		ToolbarExtender->AddToolBarExtension("Settings", EExtensionHook::After, PluginCommands, FToolBarExtensionDelegate::CreateRaw(this, &FPLUGIN_NAMEModule::AddToolbarExtension));
+		
+		LevelEditorModule.GetToolBarExtensibilityManager()->AddExtender(ToolbarExtender);
+	}
+}
+
+void FPLUGIN_NAMEModule::ShutdownModule()
+{
+	// This function may be called during shutdown to clean up your module.  For modules that support dynamic reloading,
+	// we call this function before unloading the module.
+	FPLUGIN_NAMEStyle::Shutdown();
+
+	FPLUGIN_NAMECommands::Unregister();
+}
+
+void FPLUGIN_NAMEModule::PluginButtonClicked()
+{
+	// Put your "OnButtonClicked" stuff here
+	FText DialogText = FText::Format(
+							LOCTEXT("PluginButtonDialogText", "Add code to {0} in {1} to override this button's actions"),
+							FText::FromString(TEXT("FPLUGIN_NAMEModule::PluginButtonClicked()")),
+							FText::FromString(TEXT("PLUGIN_NAME.cpp"))
+					   );
+	FMessageDialog::Open(EAppMsgType::Ok, DialogText);
+}
+
+void FPLUGIN_NAMEModule::AddMenuExtension(FMenuBuilder& Builder)
+{
+	Builder.AddMenuEntry(FPLUGIN_NAMECommands::Get().PluginAction);
+}
+
+void FPLUGIN_NAMEModule::AddToolbarExtension(FToolBarBuilder& Builder)
+{
+	Builder.AddToolBarButton(FPLUGIN_NAMECommands::Get().PluginAction);
+}
+
+#undef LOCTEXT_NAMESPACE
+	
+IMPLEMENT_MODULE(FPLUGIN_NAMEModule, PLUGIN_NAME)
\ No newline at end of file
diff --git a/Engine/Plugins/Editor/PluginBrowser/Templates/Basic/Source/PLUGIN_NAME/Private/PLUGIN_NAMEPrivatePCH.h b/Engine/Plugins/Editor/PluginBrowser/Templates/Basic/Source/PLUGIN_NAME/Private/PLUGIN_NAMEPrivatePCH.h
new file mode 100644
--- /dev/null
+++ b/Engine/Plugins/Editor/PluginBrowser/Templates/Basic/Source/PLUGIN_NAME/Private/PLUGIN_NAMEPrivatePCH.h
@@ -0,0 +1,8 @@
+// Copyright 1998-2015 Epic Games, Inc. All Rights Reserved.
+
+#include "SlateBasics.h"
+
+#include "PLUGIN_NAME.h"
+
+// You should place include statements to your module's private header files here.  You only need to
+// add includes for headers that are used in most of your module's source files though.
\ No newline at end of file
diff --git a/Engine/Plugins/Editor/PluginBrowser/Templates/Basic/Source/PLUGIN_NAME/Private/PLUGIN_NAMEStyle.cpp b/Engine/Plugins/Editor/PluginBrowser/Templates/Basic/Source/PLUGIN_NAME/Private/PLUGIN_NAMEStyle.cpp
new file mode 100644
--- /dev/null
+++ b/Engine/Plugins/Editor/PluginBrowser/Templates/Basic/Source/PLUGIN_NAME/Private/PLUGIN_NAMEStyle.cpp
@@ -0,0 +1,67 @@
+// Copyright 1998-2015 Epic Games, Inc. All Rights Reserved.
+
+#include "PLUGIN_NAMEPrivatePCH.h"
+
+#include "PLUGIN_NAMEStyle.h"
+#include "SlateGameResources.h"
+#include "IPluginManager.h"
+
+TSharedPtr< FSlateStyleSet > FPLUGIN_NAMEStyle::StyleInstance = NULL;
+
+void FPLUGIN_NAMEStyle::Initialize()
+{
+	if (!StyleInstance.IsValid())
+	{
+		StyleInstance = Create();
+		FSlateStyleRegistry::RegisterSlateStyle(*StyleInstance);
+	}
+}
+
+void FPLUGIN_NAMEStyle::Shutdown()
+{
+	FSlateStyleRegistry::UnRegisterSlateStyle(*StyleInstance);
+	ensure(StyleInstance.IsUnique());
+	StyleInstance.Reset();
+}
+
+FName FPLUGIN_NAMEStyle::GetStyleSetName()
+{
+	static FName StyleSetName(TEXT("PLUGIN_NAMEStyle"));
+	return StyleSetName;
+}
+
+#define IMAGE_BRUSH( RelativePath, ... ) FSlateImageBrush( Style->RootToContentDir( RelativePath, TEXT(".png") ), __VA_ARGS__ )
+#define BOX_BRUSH( RelativePath, ... ) FSlateBoxBrush( Style->RootToContentDir( RelativePath, TEXT(".png") ), __VA_ARGS__ )
+#define BORDER_BRUSH( RelativePath, ... ) FSlateBorderBrush( Style->RootToContentDir( RelativePath, TEXT(".png") ), __VA_ARGS__ )
+#define TTF_FONT( RelativePath, ... ) FSlateFontInfo( Style->RootToContentDir( RelativePath, TEXT(".ttf") ), __VA_ARGS__ )
+#define OTF_FONT( RelativePath, ... ) FSlateFontInfo( Style->RootToContentDir( RelativePath, TEXT(".otf") ), __VA_ARGS__ )
+
+const FVector2D Icon16x16(16.0f, 16.0f);
+const FVector2D Icon20x20(20.0f, 20.0f);
+const FVector2D Icon40x40(40.0f, 40.0f);
+
+TSharedRef< FSlateStyleSet > FPLUGIN_NAMEStyle::Create()
+{
+	TSharedRef< FSlateStyleSet > Style = MakeShareable(new FSlateStyleSet("PLUGIN_NAMEStyle"));
+	Style->SetContentRoot(IPluginManager::Get().FindPlugin("PLUGIN_NAME")->GetBaseDir() / TEXT("Resources"));
+
+	Style->Set("PLUGIN_NAME.PluginAction", new IMAGE_BRUSH(TEXT("ButtonIcon_40x"), Icon40x40));
+
+	return Style;
+}
+
+#undef IMAGE_BRUSH
+#undef BOX_BRUSH
+#undef BORDER_BRUSH
+#undef TTF_FONT
+#undef OTF_FONT
+
+void FPLUGIN_NAMEStyle::ReloadTextures()
+{
+	FSlateApplication::Get().GetRenderer()->ReloadTextureResources();
+}
+
+const ISlateStyle& FPLUGIN_NAMEStyle::Get()
+{
+	return *StyleInstance;
+}
diff --git a/Engine/Plugins/Editor/PluginBrowser/Templates/Basic/Source/PLUGIN_NAME/Public/PLUGIN_NAME.h b/Engine/Plugins/Editor/PluginBrowser/Templates/Basic/Source/PLUGIN_NAME/Public/PLUGIN_NAME.h
new file mode 100644
--- /dev/null
+++ b/Engine/Plugins/Editor/PluginBrowser/Templates/Basic/Source/PLUGIN_NAME/Public/PLUGIN_NAME.h
@@ -0,0 +1,28 @@
+// Copyright 1998-2015 Epic Games, Inc. All Rights Reserved.
+
+#pragma once
+
+#include "ModuleManager.h"
+
+class FToolBarBuilder;
+class FMenuBuilder;
+
+class FPLUGIN_NAMEModule : public IModuleInterface
+{
+public:
+
+	/** IModuleInterface implementation */
+	virtual void StartupModule() override;
+	virtual void ShutdownModule() override;
+	
+	/** This function will be bound to Command. */
+	void PluginButtonClicked();
+	
+private:
+
+	void AddToolbarExtension(FToolBarBuilder& Builder);
+	void AddMenuExtension(FMenuBuilder& Builder);
+
+private:
+	TSharedPtr<class FUICommandList> PluginCommands;
+};
\ No newline at end of file
diff --git a/Engine/Plugins/Editor/PluginBrowser/Templates/Basic/Source/PLUGIN_NAME/Public/PLUGIN_NAMECommands.cpp b/Engine/Plugins/Editor/PluginBrowser/Templates/Basic/Source/PLUGIN_NAME/Public/PLUGIN_NAMECommands.cpp
new file mode 100644
--- /dev/null
+++ b/Engine/Plugins/Editor/PluginBrowser/Templates/Basic/Source/PLUGIN_NAME/Public/PLUGIN_NAMECommands.cpp
@@ -0,0 +1,13 @@
+// Copyright 1998-2015 Epic Games, Inc. All Rights Reserved.
+
+#include "PLUGIN_NAMEPrivatePCH.h"
+#include "PLUGIN_NAMECommands.h"
+
+#define LOCTEXT_NAMESPACE "FPLUGIN_NAMEModule"
+
+void FPLUGIN_NAMECommands::RegisterCommands()
+{
+	UI_COMMAND(PluginAction, "PLUGIN_NAME", "Execute PLUGIN_NAME action", EUserInterfaceActionType::Button, FInputGesture());
+}
+
+#undef LOCTEXT_NAMESPACE
diff --git a/Engine/Plugins/Editor/PluginBrowser/Templates/Basic/Source/PLUGIN_NAME/Public/PLUGIN_NAMECommands.h b/Engine/Plugins/Editor/PluginBrowser/Templates/Basic/Source/PLUGIN_NAME/Public/PLUGIN_NAMECommands.h
new file mode 100644
--- /dev/null
+++ b/Engine/Plugins/Editor/PluginBrowser/Templates/Basic/Source/PLUGIN_NAME/Public/PLUGIN_NAMECommands.h
@@ -0,0 +1,22 @@
+// Copyright 1998-2015 Epic Games, Inc. All Rights Reserved.
+
+#pragma once
+
+#include "SlateBasics.h"
+#include "PLUGIN_NAMEStyle.h"
+
+class FPLUGIN_NAMECommands : public TCommands<FPLUGIN_NAMECommands>
+{
+public:
+
+	FPLUGIN_NAMECommands()
+		: TCommands<FPLUGIN_NAMECommands>(TEXT("PLUGIN_NAME"), NSLOCTEXT("Contexts", "PLUGIN_NAME", "PLUGIN_NAME Plugin"), NAME_None, FPLUGIN_NAMEStyle::GetStyleSetName())
+	{
+	}
+
+	// TCommands<> interface
+	virtual void RegisterCommands() override;
+
+public:
+	TSharedPtr< FUICommandInfo > PluginAction;
+};
diff --git a/Engine/Plugins/Editor/PluginBrowser/Templates/Basic/Source/PLUGIN_NAME/Public/PLUGIN_NAMEStyle.h b/Engine/Plugins/Editor/PluginBrowser/Templates/Basic/Source/PLUGIN_NAME/Public/PLUGIN_NAMEStyle.h
new file mode 100644
--- /dev/null
+++ b/Engine/Plugins/Editor/PluginBrowser/Templates/Basic/Source/PLUGIN_NAME/Public/PLUGIN_NAMEStyle.h
@@ -0,0 +1,30 @@
+// Copyright 1998-2015 Epic Games, Inc. All Rights Reserved.
+#pragma once
+
+#include "SlateBasics.h"
+
+/**  */
+class FPLUGIN_NAMEStyle
+{
+public:
+
+	static void Initialize();
+
+	static void Shutdown();
+
+	/** reloads textures used by slate renderer */
+	static void ReloadTextures();
+
+	/** @return The Slate style set for the Shooter game */
+	static const ISlateStyle& Get();
+
+	static FName GetStyleSetName();
+
+private:
+
+	static TSharedRef< class FSlateStyleSet > Create();
+
+private:
+
+	static TSharedPtr< class FSlateStyleSet > StyleInstance;
+};
\ No newline at end of file
diff --git a/Engine/Plugins/Editor/PluginBrowser/Templates/Blank/PluginModule.Build.cs.template b/Engine/Plugins/Editor/PluginBrowser/Templates/Blank/PluginModule.Build.cs.template
deleted file mode 100644
diff --git a/Engine/Plugins/Editor/PluginBrowser/Templates/Blank/PluginModule.cpp.template b/Engine/Plugins/Editor/PluginBrowser/Templates/Blank/PluginModule.cpp.template
deleted file mode 100644
diff --git a/Engine/Plugins/Editor/PluginBrowser/Templates/Blank/PluginModule.h.template b/Engine/Plugins/Editor/PluginBrowser/Templates/Blank/PluginModule.h.template
deleted file mode 100644
diff --git a/Engine/Plugins/Editor/PluginBrowser/Templates/Blank/PluginModulePCH.h.template b/Engine/Plugins/Editor/PluginBrowser/Templates/Blank/PluginModulePCH.h.template
deleted file mode 100644
diff --git a/Engine/Plugins/Editor/PluginBrowser/Templates/Blank/Source/PLUGIN_NAME/PLUGIN_NAME.Build.cs b/Engine/Plugins/Editor/PluginBrowser/Templates/Blank/Source/PLUGIN_NAME/PLUGIN_NAME.Build.cs
new file mode 100644
--- /dev/null
+++ b/Engine/Plugins/Editor/PluginBrowser/Templates/Blank/Source/PLUGIN_NAME/PLUGIN_NAME.Build.cs
@@ -0,0 +1,54 @@
+// Copyright 1998-2015 Epic Games, Inc. All Rights Reserved.
+
+using UnrealBuildTool;
+
+public class PLUGIN_NAME : ModuleRules
+{
+	public PLUGIN_NAME(TargetInfo Target)
+	{
+		
+		PublicIncludePaths.AddRange(
+			new string[] {
+				"PLUGIN_NAME/Public"
+				// ... add public include paths required here ...
+			}
+			);
+				
+		
+		PrivateIncludePaths.AddRange(
+			new string[] {
+				"PLUGIN_NAME/Private",
+				// ... add other private include paths required here ...
+			}
+			);
+			
+		
+		PublicDependencyModuleNames.AddRange(
+			new string[]
+			{
+				"Core",
+				// ... add other public dependencies that you statically link with here ...
+			}
+			);
+			
+		
+		PrivateDependencyModuleNames.AddRange(
+			new string[]
+			{
+				"CoreUObject",
+				"Engine",
+				"Slate",
+				"SlateCore",
+				// ... add private dependencies that you statically link with here ...	
+			}
+			);
+		
+		
+		DynamicallyLoadedModuleNames.AddRange(
+			new string[]
+			{
+				// ... add any modules that your module loads dynamically here ...
+			}
+			);
+	}
+}
diff --git a/Engine/Plugins/Editor/PluginBrowser/Templates/Blank/Source/PLUGIN_NAME/Private/PLUGIN_NAME.cpp b/Engine/Plugins/Editor/PluginBrowser/Templates/Blank/Source/PLUGIN_NAME/Private/PLUGIN_NAME.cpp
new file mode 100644
--- /dev/null
+++ b/Engine/Plugins/Editor/PluginBrowser/Templates/Blank/Source/PLUGIN_NAME/Private/PLUGIN_NAME.cpp
@@ -0,0 +1,20 @@
+// Copyright 1998-2015 Epic Games, Inc. All Rights Reserved.
+
+#include "PLUGIN_NAMEPrivatePCH.h"
+
+#define LOCTEXT_NAMESPACE "FPLUGIN_NAMEModule"
+
+void FPLUGIN_NAMEModule::StartupModule()
+{
+	// This code will execute after your module is loaded into memory; the exact timing is specified in the .uplugin file per-module
+}
+
+void FPLUGIN_NAMEModule::ShutdownModule()
+{
+	// This function may be called during shutdown to clean up your module.  For modules that support dynamic reloading,
+	// we call this function before unloading the module.
+}
+
+#undef LOCTEXT_NAMESPACE
+	
+IMPLEMENT_MODULE(FPLUGIN_NAMEModule, PLUGIN_NAME)
\ No newline at end of file
diff --git a/Engine/Plugins/Editor/PluginBrowser/Templates/Blank/Source/PLUGIN_NAME/Private/PLUGIN_NAMEPrivatePCH.h b/Engine/Plugins/Editor/PluginBrowser/Templates/Blank/Source/PLUGIN_NAME/Private/PLUGIN_NAMEPrivatePCH.h
new file mode 100644
--- /dev/null
+++ b/Engine/Plugins/Editor/PluginBrowser/Templates/Blank/Source/PLUGIN_NAME/Private/PLUGIN_NAMEPrivatePCH.h
@@ -0,0 +1,6 @@
+// Copyright 1998-2015 Epic Games, Inc. All Rights Reserved.
+
+#include "PLUGIN_NAME.h"
+
+// You should place include statements to your module's private header files here.  You only need to
+// add includes for headers that are used in most of your module's source files though.
\ No newline at end of file
diff --git a/Engine/Plugins/Editor/PluginBrowser/Templates/Blank/Source/PLUGIN_NAME/Public/PLUGIN_NAME.h b/Engine/Plugins/Editor/PluginBrowser/Templates/Blank/Source/PLUGIN_NAME/Public/PLUGIN_NAME.h
new file mode 100644
--- /dev/null
+++ b/Engine/Plugins/Editor/PluginBrowser/Templates/Blank/Source/PLUGIN_NAME/Public/PLUGIN_NAME.h
@@ -0,0 +1,14 @@
+// Copyright 1998-2015 Epic Games, Inc. All Rights Reserved.
+
+#pragma once
+
+#include "ModuleManager.h"
+
+class FPLUGIN_NAMEModule : public IModuleInterface
+{
+public:
+
+	/** IModuleInterface implementation */
+	virtual void StartupModule() override;
+	virtual void ShutdownModule() override;
+};
\ No newline at end of file
diff --git a/Engine/Plugins/Editor/PluginBrowser/Templates/BlueprintFunctionLibrary.cpp.template b/Engine/Plugins/Editor/PluginBrowser/Templates/BlueprintFunctionLibrary.cpp.template
deleted file mode 100644
diff --git a/Engine/Plugins/Editor/PluginBrowser/Templates/BlueprintFunctionLibrary.h.template b/Engine/Plugins/Editor/PluginBrowser/Templates/BlueprintFunctionLibrary.h.template
deleted file mode 100644
diff --git a/Engine/Plugins/Editor/PluginBrowser/Templates/BlueprintLibrary/Source/PLUGIN_NAME/PLUGIN_NAME.Build.cs b/Engine/Plugins/Editor/PluginBrowser/Templates/BlueprintLibrary/Source/PLUGIN_NAME/PLUGIN_NAME.Build.cs
new file mode 100644
--- /dev/null
+++ b/Engine/Plugins/Editor/PluginBrowser/Templates/BlueprintLibrary/Source/PLUGIN_NAME/PLUGIN_NAME.Build.cs
@@ -0,0 +1,58 @@
+// Some copyright should be here...
+
+using UnrealBuildTool;
+
+public class PLUGIN_NAME : ModuleRules
+{
+	public PLUGIN_NAME(TargetInfo Target)
+	{
+		
+		PublicIncludePaths.AddRange(
+			new string[] {
+				"PLUGIN_NAME/Public"
+				
+				// ... add public include paths required here ...
+			}
+			);
+				
+		
+		PrivateIncludePaths.AddRange(
+			new string[] {
+				"PLUGIN_NAME/Private",
+				
+				// ... add other private include paths required here ...
+			}
+			);
+			
+		
+		PublicDependencyModuleNames.AddRange(
+			new string[]
+			{
+				"Core",
+				
+				// ... add other public dependencies that you statically link with here ...
+			}
+			);
+			
+		
+		PrivateDependencyModuleNames.AddRange(
+			new string[]
+			{
+				"CoreUObject",
+				"Engine",
+				"Slate",
+				"SlateCore",
+
+				// ... add private dependencies that you statically link with here ...	
+			}
+			);
+		
+		
+		DynamicallyLoadedModuleNames.AddRange(
+			new string[]
+			{
+				// ... add any modules that your module loads dynamically here ...
+			}
+			);
+	}
+}
diff --git a/Engine/Plugins/Editor/PluginBrowser/Templates/BlueprintLibrary/Source/PLUGIN_NAME/Private/PLUGIN_NAME.cpp b/Engine/Plugins/Editor/PluginBrowser/Templates/BlueprintLibrary/Source/PLUGIN_NAME/Private/PLUGIN_NAME.cpp
new file mode 100644
--- /dev/null
+++ b/Engine/Plugins/Editor/PluginBrowser/Templates/BlueprintLibrary/Source/PLUGIN_NAME/Private/PLUGIN_NAME.cpp
@@ -0,0 +1,22 @@
+// Copyright 1998-2015 Epic Games, Inc. All Rights Reserved.
+
+#include "PLUGIN_NAMEPrivatePCH.h"
+
+#define LOCTEXT_NAMESPACE "FPLUGIN_NAMEModule"
+
+void FPLUGIN_NAMEModule::StartupModule()
+{
+	// This code will execute after your module is loaded into memory; the exact timing is specified in the .uplugin file per-module
+	
+}
+
+void FPLUGIN_NAMEModule::ShutdownModule()
+{
+	// This function may be called during shutdown to clean up your module.  For modules that support dynamic reloading,
+	// we call this function before unloading the module.
+	
+}
+
+#undef LOCTEXT_NAMESPACE
+	
+IMPLEMENT_MODULE(FPLUGIN_NAMEModule, PLUGIN_NAME)
\ No newline at end of file
diff --git a/Engine/Plugins/Editor/PluginBrowser/Templates/BlueprintLibrary/Source/PLUGIN_NAME/Private/PLUGIN_NAMEBPLibrary.cpp b/Engine/Plugins/Editor/PluginBrowser/Templates/BlueprintLibrary/Source/PLUGIN_NAME/Private/PLUGIN_NAMEBPLibrary.cpp
new file mode 100644
--- /dev/null
+++ b/Engine/Plugins/Editor/PluginBrowser/Templates/BlueprintLibrary/Source/PLUGIN_NAME/Private/PLUGIN_NAMEBPLibrary.cpp
@@ -0,0 +1,16 @@
+// Copyright 1998-2015 Epic Games, Inc. All Rights Reserved.
+
+#include "PLUGIN_NAMEPrivatePCH.h"
+#include "PLUGIN_NAMEBPLibrary.h"
+
+UPLUGIN_NAMEBPLibrary::UPLUGIN_NAMEBPLibrary(const FObjectInitializer& ObjectInitializer)
+: Super(ObjectInitializer)
+{
+
+}
+
+float UPLUGIN_NAMEBPLibrary::PLUGIN_NAMESampleFunction(float Param)
+{
+	return -1;
+}
+
diff --git a/Engine/Plugins/Editor/PluginBrowser/Templates/BlueprintLibrary/Source/PLUGIN_NAME/Private/PLUGIN_NAMEPrivatePCH.h b/Engine/Plugins/Editor/PluginBrowser/Templates/BlueprintLibrary/Source/PLUGIN_NAME/Private/PLUGIN_NAMEPrivatePCH.h
new file mode 100644
--- /dev/null
+++ b/Engine/Plugins/Editor/PluginBrowser/Templates/BlueprintLibrary/Source/PLUGIN_NAME/Private/PLUGIN_NAMEPrivatePCH.h
@@ -0,0 +1,6 @@
+// Copyright 1998-2015 Epic Games, Inc. All Rights Reserved.
+
+#include "PLUGIN_NAME.h"
+
+// You should place include statements to your module's private header files here.  You only need to
+// add includes for headers that are used in most of your module's source files though.
\ No newline at end of file
diff --git a/Engine/Plugins/Editor/PluginBrowser/Templates/BlueprintLibrary/Source/PLUGIN_NAME/Public/PLUGIN_NAME.h b/Engine/Plugins/Editor/PluginBrowser/Templates/BlueprintLibrary/Source/PLUGIN_NAME/Public/PLUGIN_NAME.h
new file mode 100644
--- /dev/null
+++ b/Engine/Plugins/Editor/PluginBrowser/Templates/BlueprintLibrary/Source/PLUGIN_NAME/Public/PLUGIN_NAME.h
@@ -0,0 +1,14 @@
+// Copyright 1998-2015 Epic Games, Inc. All Rights Reserved.
+
+#pragma once
+
+#include "ModuleManager.h"
+
+class FPLUGIN_NAMEModule : public IModuleInterface
+{
+public:
+
+	/** IModuleInterface implementation */
+	virtual void StartupModule() override;
+	virtual void ShutdownModule() override;
+};
\ No newline at end of file
diff --git a/Engine/Plugins/Editor/PluginBrowser/Templates/BlueprintLibrary/Source/PLUGIN_NAME/Public/PLUGIN_NAMEBPLibrary.h b/Engine/Plugins/Editor/PluginBrowser/Templates/BlueprintLibrary/Source/PLUGIN_NAME/Public/PLUGIN_NAMEBPLibrary.h
new file mode 100644
--- /dev/null
+++ b/Engine/Plugins/Editor/PluginBrowser/Templates/BlueprintLibrary/Source/PLUGIN_NAME/Public/PLUGIN_NAMEBPLibrary.h
@@ -0,0 +1,32 @@
+// Copyright 1998-2015 Epic Games, Inc. All Rights Reserved.
+
+#pragma once
+
+#include "Engine.h"
+#include "PLUGIN_NAMEBPLibrary.generated.h"
+
+/* 
+*	Function library class.
+*	Each function in it is expected to be static and represents blueprint node that can be called in any blueprint.
+*
+*	When declaring function you can define metadata for the node. Key function specifiers will be BlueprintPure and BlueprintCallable.
+*	BlueprintPure - means the function does not affect the owning object in any way and thus creates a node without Exec pins.
+*	BlueprintCallable - makes a function which can be executed in Blueprints - Thus it has Exec pins.
+*	DisplayName - full name of the node, shown when you mouse over the node and in the blueprint drop down menu.
+*				Its lets you name the node using characters not allowed in C++ function names.
+*	CompactNodeTitle - the word(s) that appear on the node.
+*	Keywords -	the list of keywords that helps you to find node when you search for it using Blueprint drop-down menu. 
+*				Good example is "Print String" node which you can find also by using keyword "log".
+*	Category -	the category your node will be under in the Blueprint drop-down menu.
+*
+*	For more info on custom blueprint nodes visit documentation:
+*	https://wiki.unrealengine.com/Custom_Blueprint_Node_Creation
+*/
+UCLASS()
+class UPLUGIN_NAMEBPLibrary : public UBlueprintFunctionLibrary
+{
+	GENERATED_UCLASS_BODY()
+
+	UFUNCTION(BlueprintCallable, meta = (DisplayName = "Execute Sample function", Keywords = "PLUGIN_NAME sample test testing"), Category = "PLUGIN_NAMETesting")
+	static float PLUGIN_NAMESampleFunction(float Param);
+};
diff --git a/Engine/Plugins/Editor/PluginBrowser/Templates/EdMode.cpp.template b/Engine/Plugins/Editor/PluginBrowser/Templates/EdMode.cpp.template
deleted file mode 100644
diff --git a/Engine/Plugins/Editor/PluginBrowser/Templates/EdMode.h.template b/Engine/Plugins/Editor/PluginBrowser/Templates/EdMode.h.template
deleted file mode 100644
diff --git a/Engine/Plugins/Editor/PluginBrowser/Templates/EditorMode/Source/PLUGIN_NAME/PLUGIN_NAME.Build.cs b/Engine/Plugins/Editor/PluginBrowser/Templates/EditorMode/Source/PLUGIN_NAME/PLUGIN_NAME.Build.cs
new file mode 100644
--- /dev/null
+++ b/Engine/Plugins/Editor/PluginBrowser/Templates/EditorMode/Source/PLUGIN_NAME/PLUGIN_NAME.Build.cs
@@ -0,0 +1,57 @@
+// Copyright 1998-2015 Epic Games, Inc. All Rights Reserved.
+
+using UnrealBuildTool;
+
+public class PLUGIN_NAME : ModuleRules
+{
+	public PLUGIN_NAME(TargetInfo Target)
+	{
+		
+		PublicIncludePaths.AddRange(
+			new string[] {
+				"PLUGIN_NAME/Public"
+				// ... add public include paths required here ...
+			}
+			);
+				
+		
+		PrivateIncludePaths.AddRange(
+			new string[] {
+				"PLUGIN_NAME/Private",
+				// ... add other private include paths required here ...
+			}
+			);
+			
+		
+		PublicDependencyModuleNames.AddRange(
+			new string[]
+			{
+				"Core",
+				// ... add other public dependencies that you statically link with here ...
+			}
+			);
+			
+		
+		PrivateDependencyModuleNames.AddRange(
+			new string[]
+			{
+				"CoreUObject",
+				"Engine",
+				"Slate",
+				"SlateCore",
+				"InputCore",
+				"UnrealEd",
+				"LevelEditor",
+				// ... add private dependencies that you statically link with here ...	
+			}
+			);
+		
+		
+		DynamicallyLoadedModuleNames.AddRange(
+			new string[]
+			{
+				// ... add any modules that your module loads dynamically here ...
+			}
+			);
+	}
+}
diff --git a/Engine/Plugins/Editor/PluginBrowser/Templates/EditorMode/Source/PLUGIN_NAME/Private/PLUGIN_NAME.cpp b/Engine/Plugins/Editor/PluginBrowser/Templates/EditorMode/Source/PLUGIN_NAME/Private/PLUGIN_NAME.cpp
new file mode 100644
--- /dev/null
+++ b/Engine/Plugins/Editor/PluginBrowser/Templates/EditorMode/Source/PLUGIN_NAME/Private/PLUGIN_NAME.cpp
@@ -0,0 +1,23 @@
+// Copyright 1998-2015 Epic Games, Inc. All Rights Reserved.
+
+#include "PLUGIN_NAMEPrivatePCH.h"
+#include "PLUGIN_NAMEEdMode.h"
+
+#define LOCTEXT_NAMESPACE "FPLUGIN_NAMEModule"
+
+void FPLUGIN_NAMEModule::StartupModule()
+{
+	// This code will execute after your module is loaded into memory; the exact timing is specified in the .uplugin file per-module
+	FEditorModeRegistry::Get().RegisterMode<FPLUGIN_NAMEEdMode>(FPLUGIN_NAMEEdMode::EM_PLUGIN_NAMEEdModeId, LOCTEXT("PLUGIN_NAMEEdModeName", "PLUGIN_NAMEEdMode"), FSlateIcon(), true);
+}
+
+void FPLUGIN_NAMEModule::ShutdownModule()
+{
+	// This function may be called during shutdown to clean up your module.  For modules that support dynamic reloading,
+	// we call this function before unloading the module.
+	FEditorModeRegistry::Get().UnregisterMode(FPLUGIN_NAMEEdMode::EM_PLUGIN_NAMEEdModeId);
+}
+
+#undef LOCTEXT_NAMESPACE
+	
+IMPLEMENT_MODULE(FPLUGIN_NAMEModule, PLUGIN_NAME)
\ No newline at end of file
diff --git a/Engine/Plugins/Editor/PluginBrowser/Templates/EditorMode/Source/PLUGIN_NAME/Private/PLUGIN_NAMEEdMode.cpp b/Engine/Plugins/Editor/PluginBrowser/Templates/EditorMode/Source/PLUGIN_NAME/Private/PLUGIN_NAMEEdMode.cpp
new file mode 100644
--- /dev/null
+++ b/Engine/Plugins/Editor/PluginBrowser/Templates/EditorMode/Source/PLUGIN_NAME/Private/PLUGIN_NAMEEdMode.cpp
@@ -0,0 +1,50 @@
+// Copyright 1998-2015 Epic Games, Inc. All Rights Reserved.
+
+#include "PLUGIN_NAMEPrivatePCH.h"
+#include "PLUGIN_NAMEEdMode.h"
+#include "PLUGIN_NAMEEdModeToolkit.h"
+#include "Toolkits/ToolkitManager.h"
+
+const FEditorModeID FPLUGIN_NAMEEdMode::EM_PLUGIN_NAMEEdModeId = TEXT("EM_PLUGIN_NAMEEdMode");
+
+FPLUGIN_NAMEEdMode::FPLUGIN_NAMEEdMode()
+{
+
+}
+
+FPLUGIN_NAMEEdMode::~FPLUGIN_NAMEEdMode()
+{
+
+}
+
+void FPLUGIN_NAMEEdMode::Enter()
+{
+	FEdMode::Enter();
+
+	if (!Toolkit.IsValid() && UsesToolkits())
+	{
+		Toolkit = MakeShareable(new FPLUGIN_NAMEEdModeToolkit);
+		Toolkit->Init(Owner->GetToolkitHost());
+	}
+}
+
+void FPLUGIN_NAMEEdMode::Exit()
+{
+	if (Toolkit.IsValid())
+	{
+		FToolkitManager::Get().CloseToolkit(Toolkit.ToSharedRef());
+		Toolkit.Reset();
+	}
+
+	// Call base Exit method to ensure proper cleanup
+	FEdMode::Exit();
+}
+
+bool FPLUGIN_NAMEEdMode::UsesToolkits() const
+{
+	return true;
+}
+
+
+
+
diff --git a/Engine/Plugins/Editor/PluginBrowser/Templates/EditorMode/Source/PLUGIN_NAME/Private/PLUGIN_NAMEEdModeToolkit.cpp b/Engine/Plugins/Editor/PluginBrowser/Templates/EditorMode/Source/PLUGIN_NAME/Private/PLUGIN_NAMEEdModeToolkit.cpp
new file mode 100644
--- /dev/null
+++ b/Engine/Plugins/Editor/PluginBrowser/Templates/EditorMode/Source/PLUGIN_NAME/Private/PLUGIN_NAMEEdModeToolkit.cpp
@@ -0,0 +1,115 @@
+// Copyright 1998-2015 Epic Games, Inc. All Rights Reserved.
+
+#include "PLUGIN_NAMEPrivatePCH.h"
+#include "PLUGIN_NAMEEdMode.h"
+#include "PLUGIN_NAMEEdModeToolkit.h"
+
+#define LOCTEXT_NAMESPACE "FPLUGIN_NAMEEdModeToolkit"
+
+FPLUGIN_NAMEEdModeToolkit::FPLUGIN_NAMEEdModeToolkit()
+{
+	struct Locals
+	{
+		static bool IsWidgetEnabled()
+		{
+			return GEditor->GetSelectedActors()->Num() != 0;
+		}
+
+		static FReply OnButtonClick(FVector InOffset)
+		{
+			USelection* SelectedActors = GEditor->GetSelectedActors();
+
+			// Let editor know that we're about to do something that we want to undo/redo
+			GEditor->BeginTransaction(LOCTEXT("MoveActorsTransactionName", "MoveActors"));
+
+			// For each selected actor
+			for (FSelectionIterator Iter(*SelectedActors); Iter; ++Iter)
+			{
+				if (AActor* LevelActor = Cast<AActor>(*Iter))
+				{
+					// Register actor in opened transaction (undo/redo)
+					LevelActor->Modify();
+					// Move actor to given location
+					LevelActor->TeleportTo(LevelActor->GetActorLocation() + InOffset, FRotator(0, 0, 0));
+				}
+			}
+
+			// We're done moving actors so close transaction
+			GEditor->EndTransaction();
+
+			return FReply::Handled();
+		}
+
+		static TSharedRef<SWidget> MakeButton(FText InLabel, const FVector InOffset)
+		{
+			return SNew(SButton)
+				.Text(InLabel)
+				.OnClicked_Static(&Locals::OnButtonClick, InOffset);
+		}
+	};
+
+	const float Factor = 256.0f;
+
+	SAssignNew(ToolkitWidget, SBorder)
+		.HAlign(HAlign_Center)
+		.Padding(25)
+		.IsEnabled_Static(&Locals::IsWidgetEnabled)
+		[
+			SNew(SVerticalBox)
+			+ SVerticalBox::Slot()
+			.AutoHeight()
+			.HAlign(HAlign_Center)
+			.Padding(50)
+			[
+				SNew(STextBlock)
+				.AutoWrapText(true)
+				.Text(LOCTEXT("HelperLabel", "Select some actors and move them around using buttons below"))
+			]
+			+ SVerticalBox::Slot()
+				.HAlign(HAlign_Center)
+				.AutoHeight()
+				[
+					Locals::MakeButton(LOCTEXT("UpButtonLabel", "Up"), FVector(0, 0, Factor))
+				]
+			+ SVerticalBox::Slot()
+				.HAlign(HAlign_Center)
+				.AutoHeight()
+				[
+					SNew(SHorizontalBox)
+					+ SHorizontalBox::Slot()
+					.AutoWidth()
+					[
+						Locals::MakeButton(LOCTEXT("LeftButtonLabel", "Left"), FVector(0, -Factor, 0))
+					]
+					+ SHorizontalBox::Slot()
+						.AutoWidth()
+						[
+							Locals::MakeButton(LOCTEXT("RightButtonLabel", "Right"), FVector(0, Factor, 0))
+						]
+				]
+			+ SVerticalBox::Slot()
+				.HAlign(HAlign_Center)
+				.AutoHeight()
+				[
+					Locals::MakeButton(LOCTEXT("DownButtonLabel", "Down"), FVector(0, 0, -Factor))
+				]
+
+		];
+}
+
+FName FPLUGIN_NAMEEdModeToolkit::GetToolkitFName() const
+{
+	return FName("PLUGIN_NAMEEdMode");
+}
+
+FText FPLUGIN_NAMEEdModeToolkit::GetBaseToolkitName() const
+{
+	return NSLOCTEXT("PLUGIN_NAMEEdModeToolkit", "DisplayName", "PLUGIN_NAMEEdMode Tool");
+}
+
+class FEdMode* FPLUGIN_NAMEEdModeToolkit::GetEditorMode() const
+{
+	return GLevelEditorModeTools().GetActiveMode(FPLUGIN_NAMEEdMode::EM_PLUGIN_NAMEEdModeId);
+}
+
+#undef LOCTEXT_NAMESPACE
diff --git a/Engine/Plugins/Editor/PluginBrowser/Templates/EditorMode/Source/PLUGIN_NAME/Private/PLUGIN_NAMEPrivatePCH.h b/Engine/Plugins/Editor/PluginBrowser/Templates/EditorMode/Source/PLUGIN_NAME/Private/PLUGIN_NAMEPrivatePCH.h
new file mode 100644
--- /dev/null
+++ b/Engine/Plugins/Editor/PluginBrowser/Templates/EditorMode/Source/PLUGIN_NAME/Private/PLUGIN_NAMEPrivatePCH.h
@@ -0,0 +1,6 @@
+// Copyright 1998-2015 Epic Games, Inc. All Rights Reserved.
+
+#include "PLUGIN_NAME.h"
+
+// You should place include statements to your module's private header files here.  You only need to
+// add includes for headers that are used in most of your module's source files though.
\ No newline at end of file
diff --git a/Engine/Plugins/Editor/PluginBrowser/Templates/EditorMode/Source/PLUGIN_NAME/Public/PLUGIN_NAME.h b/Engine/Plugins/Editor/PluginBrowser/Templates/EditorMode/Source/PLUGIN_NAME/Public/PLUGIN_NAME.h
new file mode 100644
--- /dev/null
+++ b/Engine/Plugins/Editor/PluginBrowser/Templates/EditorMode/Source/PLUGIN_NAME/Public/PLUGIN_NAME.h
@@ -0,0 +1,14 @@
+// Copyright 1998-2015 Epic Games, Inc. All Rights Reserved.
+
+#pragma once
+
+#include "ModuleManager.h"
+
+class FPLUGIN_NAMEModule : public IModuleInterface
+{
+public:
+
+	/** IModuleInterface implementation */
+	virtual void StartupModule() override;
+	virtual void ShutdownModule() override;
+};
\ No newline at end of file
diff --git a/Engine/Plugins/Editor/PluginBrowser/Templates/EditorMode/Source/PLUGIN_NAME/Public/PLUGIN_NAMEEdMode.h b/Engine/Plugins/Editor/PluginBrowser/Templates/EditorMode/Source/PLUGIN_NAME/Public/PLUGIN_NAMEEdMode.h
new file mode 100644
--- /dev/null
+++ b/Engine/Plugins/Editor/PluginBrowser/Templates/EditorMode/Source/PLUGIN_NAME/Public/PLUGIN_NAMEEdMode.h
@@ -0,0 +1,24 @@
+// Copyright 1998-2015 Epic Games, Inc. All Rights Reserved.
+
+#pragma once
+
+#include "UnrealEd.h" 
+#include "Editor.h"
+
+class FPLUGIN_NAMEEdMode : public FEdMode
+{
+public:
+	const static FEditorModeID EM_PLUGIN_NAMEEdModeId;
+public:
+	FPLUGIN_NAMEEdMode();
+	virtual ~FPLUGIN_NAMEEdMode();
+
+	// FEdMode interface
+	virtual void Enter() override;
+	virtual void Exit() override;
+	//virtual void Tick(FEditorViewportClient* ViewportClient, float DeltaTime) override;
+	//virtual void Render(const FSceneView* View, FViewport* Viewport, FPrimitiveDrawInterface* PDI) override;
+	//virtual void ActorSelectionChangeNotify() override;
+	bool UsesToolkits() const override;
+	// End of FEdMode interface
+};
diff --git a/Engine/Plugins/Editor/PluginBrowser/Templates/EditorMode/Source/PLUGIN_NAME/Public/PLUGIN_NAMEEdModeToolkit.h b/Engine/Plugins/Editor/PluginBrowser/Templates/EditorMode/Source/PLUGIN_NAME/Public/PLUGIN_NAMEEdModeToolkit.h
new file mode 100644
--- /dev/null
+++ b/Engine/Plugins/Editor/PluginBrowser/Templates/EditorMode/Source/PLUGIN_NAME/Public/PLUGIN_NAMEEdModeToolkit.h
@@ -0,0 +1,23 @@
+// Copyright 1998-2015 Epic Games, Inc. All Rights Reserved.
+
+#pragma once
+
+#include "Editor/UnrealEd/Public/Toolkits/BaseToolkit.h"
+
+
+class FPLUGIN_NAMEEdModeToolkit : public FModeToolkit
+{
+public:
+
+	FPLUGIN_NAMEEdModeToolkit();
+
+	/** IToolkit interface */
+	virtual FName GetToolkitFName() const override;
+	virtual FText GetBaseToolkitName() const override;
+	virtual class FEdMode* GetEditorMode() const override;
+	virtual TSharedPtr<class SWidget> GetInlineContent() const override { return ToolkitWidget; }
+
+private:
+
+	TSharedPtr<SWidget> ToolkitWidget;
+};
diff --git a/Engine/Plugins/Editor/PluginBrowser/Templates/ThirdPartyLibrary/Source/PLUGIN_NAME/PLUGIN_NAME.Build.cs b/Engine/Plugins/Editor/PluginBrowser/Templates/ThirdPartyLibrary/Source/PLUGIN_NAME/PLUGIN_NAME.Build.cs
--- a/Engine/Plugins/Editor/PluginBrowser/Templates/ThirdPartyLibrary/Source/PLUGIN_NAME/PLUGIN_NAME.Build.cs
+++ b/Engine/Plugins/Editor/PluginBrowser/Templates/ThirdPartyLibrary/Source/PLUGIN_NAME/PLUGIN_NAME.Build.cs
@@ -1,4 +1,4 @@
-// Some copyright should be here...
+// Copyright 1998-2015 Epic Games, Inc. All Rights Reserved.
 
 using UnrealBuildTool;
 
diff --git a/Engine/Plugins/Editor/PluginBrowser/Templates/ThirdPartyLibrary/Source/PLUGIN_NAME/Private/PLUGIN_NAME.cpp b/Engine/Plugins/Editor/PluginBrowser/Templates/ThirdPartyLibrary/Source/PLUGIN_NAME/Private/PLUGIN_NAME.cpp
--- a/Engine/Plugins/Editor/PluginBrowser/Templates/ThirdPartyLibrary/Source/PLUGIN_NAME/Private/PLUGIN_NAME.cpp
+++ b/Engine/Plugins/Editor/PluginBrowser/Templates/ThirdPartyLibrary/Source/PLUGIN_NAME/Private/PLUGIN_NAME.cpp
@@ -1,4 +1,4 @@
-// Some copyright should be here...
+// Copyright 1998-2015 Epic Games, Inc. All Rights Reserved.
 
 #include "PLUGIN_NAMEPrivatePCH.h"
 #include "Core.h"
diff --git a/Engine/Plugins/Editor/PluginBrowser/Templates/ThirdPartyLibrary/Source/PLUGIN_NAME/Private/PLUGIN_NAMEPrivatePCH.h b/Engine/Plugins/Editor/PluginBrowser/Templates/ThirdPartyLibrary/Source/PLUGIN_NAME/Private/PLUGIN_NAMEPrivatePCH.h
--- a/Engine/Plugins/Editor/PluginBrowser/Templates/ThirdPartyLibrary/Source/PLUGIN_NAME/Private/PLUGIN_NAMEPrivatePCH.h
+++ b/Engine/Plugins/Editor/PluginBrowser/Templates/ThirdPartyLibrary/Source/PLUGIN_NAME/Private/PLUGIN_NAMEPrivatePCH.h
@@ -1,4 +1,4 @@
-// Some copyright should be here...
+// Copyright 1998-2015 Epic Games, Inc. All Rights Reserved.
 
 #include "PLUGIN_NAME.h"
 
diff --git a/Engine/Plugins/Editor/PluginBrowser/Templates/ThirdPartyLibrary/Source/PLUGIN_NAME/Public/PLUGIN_NAME.h b/Engine/Plugins/Editor/PluginBrowser/Templates/ThirdPartyLibrary/Source/PLUGIN_NAME/Public/PLUGIN_NAME.h
--- a/Engine/Plugins/Editor/PluginBrowser/Templates/ThirdPartyLibrary/Source/PLUGIN_NAME/Public/PLUGIN_NAME.h
+++ b/Engine/Plugins/Editor/PluginBrowser/Templates/ThirdPartyLibrary/Source/PLUGIN_NAME/Public/PLUGIN_NAME.h
@@ -1,4 +1,4 @@
-// Some copyright should be here...
+// Copyright 1998-2015 Epic Games, Inc. All Rights Reserved.
 
 #pragma once
 
diff --git a/Engine/Plugins/Editor/PluginBrowser/Templates/Toolkit.cpp.template b/Engine/Plugins/Editor/PluginBrowser/Templates/Toolkit.cpp.template
deleted file mode 100644
diff --git a/Engine/Plugins/Editor/PluginBrowser/Templates/Toolkit.h.template b/Engine/Plugins/Editor/PluginBrowser/Templates/Toolkit.h.template
deleted file mode 100644
