Fix for crash when editing plugin metadata in the Plugins Viewer.

diff --git a/Engine/Plugins/Editor/PluginBrowser/Source/PluginBrowser/Private/SPluginCategoryTree.cpp b/Engine/Plugins/Editor/PluginBrowser/Source/PluginBrowser/Private/SPluginCategoryTree.cpp
--- a/Engine/Plugins/Editor/PluginBrowser/Source/PluginBrowser/Private/SPluginCategoryTree.cpp
+++ b/Engine/Plugins/Editor/PluginBrowser/Source/PluginBrowser/Private/SPluginCategoryTree.cpp
@@ -52,6 +52,16 @@ SPluginBrowser& SPluginCategoryTree::GetOwner()
 }
 
 
+static void ResetCategories(TArray<TSharedPtr<FPluginCategory>>& Categories)
+{
+	for (const TSharedPtr<FPluginCategory>& Category : Categories)
+	{
+		ResetCategories(Category->SubCategories);
+		Category->Plugins.Reset();
+		Category->SubCategories.Reset();
+	}
+}
+
 void SPluginCategoryTree::RebuildAndFilterCategoryTree()
 {
 	// Get a plugin from the currently selected category, so we can track it if it's removed
@@ -66,14 +76,7 @@ void SPluginCategoryTree::RebuildAndFilterCategoryTree()
 	}
 
 	// Clear the list of plugins in each current category
-	for(TSharedPtr<FPluginCategory> RootCategory: RootCategories)
-	{
-		for(TSharedPtr<FPluginCategory> Category: RootCategory->SubCategories)
-		{
-			Category->Plugins.Reset();
-			Category->SubCategories.Reset();
-		}
-	}
+	ResetCategories(RootCategories);
 
 	// Add all the known plugins into categories
 	TSharedPtr<FPluginCategory> SelectCategory;
