UE-21327 - Adding "Marketplace URL" field for custom plugins in editor.

diff --git a/Engine/Plugins/Editor/PluginBrowser/Source/PluginBrowser/Private/PluginMetadataObject.cpp b/Engine/Plugins/Editor/PluginBrowser/Source/PluginBrowser/Private/PluginMetadataObject.cpp
--- a/Engine/Plugins/Editor/PluginBrowser/Source/PluginBrowser/Private/PluginMetadataObject.cpp
+++ b/Engine/Plugins/Editor/PluginBrowser/Source/PluginBrowser/Private/PluginMetadataObject.cpp
@@ -22,6 +22,7 @@ void UPluginMetadataObject::PopulateFromDescriptor(const FPluginDescriptor& InDe
 	CreatedBy = InDescriptor.CreatedBy;
 	CreatedByURL = InDescriptor.CreatedByURL;
 	DocsURL = InDescriptor.DocsURL;
+	MarketplaceURL = InDescriptor.MarketplaceURL;
 	SupportURL = InDescriptor.SupportURL;
 	bCanContainContent = InDescriptor.bCanContainContent;
 	bIsBetaVersion = InDescriptor.bIsBetaVersion;
@@ -37,6 +38,7 @@ void UPluginMetadataObject::CopyIntoDescriptor(FPluginDescriptor& OutDescriptor)
 	OutDescriptor.CreatedBy = CreatedBy;
 	OutDescriptor.CreatedByURL = CreatedByURL;
 	OutDescriptor.DocsURL = DocsURL;
+	OutDescriptor.MarketplaceURL = MarketplaceURL;
 	OutDescriptor.SupportURL = SupportURL;
 	OutDescriptor.bCanContainContent = bCanContainContent;
 	OutDescriptor.bIsBetaVersion = bIsBetaVersion;
diff --git a/Engine/Plugins/Editor/PluginBrowser/Source/PluginBrowser/Private/PluginMetadataObject.h b/Engine/Plugins/Editor/PluginBrowser/Source/PluginBrowser/Private/PluginMetadataObject.h
--- a/Engine/Plugins/Editor/PluginBrowser/Source/PluginBrowser/Private/PluginMetadataObject.h
+++ b/Engine/Plugins/Editor/PluginBrowser/Source/PluginBrowser/Private/PluginMetadataObject.h
@@ -56,6 +56,10 @@ public:
 	UPROPERTY(EditAnywhere, Category = "Details")
 	FString DocsURL;
 
+	/** Marketplace URL string. */
+	UPROPERTY(EditAnywhere, Category = "Details")
+	FString MarketplaceURL;
+
 	/** Support URL/email for this plugin. Email addresses must be prefixed with 'mailto:' */
 	UPROPERTY(EditAnywhere, Category = "Details")
 	FString SupportURL;
