Remove subdirectory option from BuildPlugin and have Editor add it instead

diff --git a/Engine/Plugins/Editor/PluginBrowser/Source/PluginBrowser/Private/SPluginTile.cpp b/Engine/Plugins/Editor/PluginBrowser/Source/PluginBrowser/Private/SPluginTile.cpp
--- a/Engine/Plugins/Editor/PluginBrowser/Source/PluginBrowser/Private/SPluginTile.cpp
+++ b/Engine/Plugins/Editor/PluginBrowser/Source/PluginBrowser/Private/SPluginTile.cpp
@@ -592,7 +592,8 @@ void SPluginTile::OnPackagePlugin()
 	// Ensure path is full rather than relative (for macs)
 	FString DescriptorFilename = Plugin->GetDescriptorFileName();
 	FString DescriptorFullPath = FPaths::ConvertRelativePathToFull(DescriptorFilename);
-	FString CommandLine = FString::Printf(TEXT("BuildPlugin -Rocket -Plugin=\"%s\" -Package=\"%s\" -CreateSubFolder"), *DescriptorFullPath, *OutputDirectory);
+	OutputDirectory = FPaths::Combine(OutputDirectory, Plugin->GetName());
+	FString CommandLine = FString::Printf(TEXT("BuildPlugin -Rocket -Plugin=\"%s\" -Package=\"%s\""), *DescriptorFullPath, *OutputDirectory);
 
 #if PLATFORM_WINDOWS
 	FText PlatformName = LOCTEXT("PlatformName_Windows", "Windows");
diff --git a/Engine/Source/Programs/AutomationTool/Scripts/BuildPluginCommand.Automation.cs b/Engine/Source/Programs/AutomationTool/Scripts/BuildPluginCommand.Automation.cs
--- a/Engine/Source/Programs/AutomationTool/Scripts/BuildPluginCommand.Automation.cs
+++ b/Engine/Source/Programs/AutomationTool/Scripts/BuildPluginCommand.Automation.cs
@@ -49,13 +49,6 @@ class BuildPlugin : BuildCommand
 			throw new AutomationException("Output directory for packaged plugin must be outside engine directory");
 		}
 
-		// Whether to create a new sub folder with the name of the plugin to prevent deletion of everything in the package directory
-		bool bCreateSubFolder = ParseParam("CreateSubFolder");
-		if (bCreateSubFolder)
-		{
-			PackageDir = DirectoryReference.Combine(PackageDir, PluginFile.GetFileNameWithoutExtension());
-		}
-
 		// Clear the output directory of existing stuff
 		if (PackageDir.Exists())
 		{
