PR #4382: Prevent error log due to non-existing plugin directory (Contributed by [REDACTED_PERSON])

diff --git a/Engine/Plugins/Editor/PluginBrowser/Source/PluginBrowser/Private/SPluginBrowser.cpp b/Engine/Plugins/Editor/PluginBrowser/Source/PluginBrowser/Private/SPluginBrowser.cpp
--- a/Engine/Plugins/Editor/PluginBrowser/Source/PluginBrowser/Private/SPluginBrowser.cpp
+++ b/Engine/Plugins/Editor/PluginBrowser/Source/PluginBrowser/Private/SPluginBrowser.cpp
@@ -46,13 +46,20 @@ void SPluginBrowser::Construct( const FArguments& Args )
 	}
 	if(FApp::HasProjectName())
 	{
-		WatchDirectoryNames.Add(FPaths::ProjectPluginsDir());
+		const FString ProjectPluginsDir = FPaths::ProjectPluginsDir();
+		if (FPaths::DirectoryExists(ProjectPluginsDir))
+		{
+			WatchDirectoryNames.Add(ProjectPluginsDir);
+		}
 		const FProjectDescriptor* Project = IProjectManager::Get().GetCurrentProject();
 		if (Project != nullptr)
 		{
 			for (const FString& Path : Project->GetAdditionalPluginDirectories())
 			{
-				WatchDirectoryNames.Add(Path);
+				if (FPaths::DirectoryExists(Path))
+				{
+					WatchDirectoryNames.Add(Path);
+				}
 			}
 		}
 	}
