Change plugins to create runtime modules by default.

diff --git a/Engine/Plugins/Editor/PluginBrowser/Source/PluginBrowser/Private/DefaultPluginWizardDefinition.cpp b/Engine/Plugins/Editor/PluginBrowser/Source/PluginBrowser/Private/DefaultPluginWizardDefinition.cpp
--- a/Engine/Plugins/Editor/PluginBrowser/Source/PluginBrowser/Private/DefaultPluginWizardDefinition.cpp
+++ b/Engine/Plugins/Editor/PluginBrowser/Source/PluginBrowser/Private/DefaultPluginWizardDefinition.cpp
@@ -39,12 +39,12 @@ void FDefaultPluginWizardDefinition::PopulateTemplatesSource()
 	if (!bIsContentOnlyProject)
 	{
 		// Insert the blank template to make sure it appears before the content only template.
-		TemplateDefinitions.Insert(MakeShareable(new FPluginTemplateDescription(BlankTemplateName, BlankDescription, TEXT("Blank"), true, EHostType::Developer)), 0);
+		TemplateDefinitions.Insert(MakeShareable(new FPluginTemplateDescription(BlankTemplateName, BlankDescription, TEXT("Blank"), true, EHostType::Runtime)), 0);
 		TemplateDefinitions.Add(MakeShareable(new FPluginTemplateDescription(BlueprintLibTemplateName, BlueprintLibDescription, TEXT("BlueprintLibrary"), true, EHostType::Runtime, ELoadingPhase::PreLoadingScreen)));
 		TemplateDefinitions.Add(MakeShareable(new FPluginTemplateDescription(BasicTemplateName, BasicDescription, TEXT("Basic"), false, EHostType::Editor)));
 		TemplateDefinitions.Add(MakeShareable(new FPluginTemplateDescription(AdvancedTemplateName, AdvancedDescription, TEXT("Advanced"), false, EHostType::Editor)));
 		TemplateDefinitions.Add(MakeShareable(new FPluginTemplateDescription(EditorModeTemplateName, EditorModeDescription, TEXT("EditorMode"), false, EHostType::Editor)));
-		TemplateDefinitions.Add(MakeShareable(new FPluginTemplateDescription(ThirdPartyTemplateName, ThirdPartyDescription, TEXT("ThirdPartyLibrary"), true, EHostType::Developer)));
+		TemplateDefinitions.Add(MakeShareable(new FPluginTemplateDescription(ThirdPartyTemplateName, ThirdPartyDescription, TEXT("ThirdPartyLibrary"), true, EHostType::Runtime)));
 	}
 }
 
@@ -120,7 +120,7 @@ bool FDefaultPluginWizardDefinition::GetPluginIconPath(FString& OutIconPath) con
 
 EHostType::Type FDefaultPluginWizardDefinition::GetPluginModuleDescriptor() const
 {
-	EHostType::Type ModuleDescriptorType = EHostType::Developer;
+	EHostType::Type ModuleDescriptorType = EHostType::Runtime;
 
 	if (CurrentTemplateDefinition.IsValid())
 	{
