Fix new plugins created via the new plugin wizard failing to compile.

diff --git a/Engine/Plugins/Editor/PluginBrowser/Source/PluginBrowser/Private/SNewPluginWizard.cpp b/Engine/Plugins/Editor/PluginBrowser/Source/PluginBrowser/Private/SNewPluginWizard.cpp
--- a/Engine/Plugins/Editor/PluginBrowser/Source/PluginBrowser/Private/SNewPluginWizard.cpp
+++ b/Engine/Plugins/Editor/PluginBrowser/Source/PluginBrowser/Private/SNewPluginWizard.cpp
@@ -743,7 +743,7 @@ FReply SNewPluginWizard::OnCreatePluginClicked()
 	if (bSucceeded && bHasModules)
 	{
 		FString ProjectFileName = FPaths::GetProjectFilePath();
-		FString Arguments = FString::Printf(TEXT("%s %s -TargetType=Editor -Module=%s -EnablePlugin=\"%s\" -Project=\"%s\" -Progress -NoHotReloadFromIDE"), FModuleManager::Get().GetUBTConfiguration(), FPlatformMisc::GetUBTPlatform(), *AutoPluginName, *PluginModuleName, *ProjectFileName);
+		FString Arguments = FString::Printf(TEXT("%s %s -TargetType=Editor -Plugin=\"%s\" -Project=\"%s\" -Progress -NoHotReloadFromIDE"), FModuleManager::Get().GetUBTConfiguration(), FPlatformMisc::GetUBTPlatform(), *UPluginFilePath, *ProjectFileName);
 		if (!FDesktopPlatformModule::Get()->RunUnrealBuildTool(LOCTEXT("Compiling", "Compiling..."), FPaths::RootDir(), Arguments, GWarn))
 		{
 			PopErrorNotification(LOCTEXT("FailedToCompile", "Failed to compile source code."));
