Convert project file path to absolute before passing on the command line to UBT.

diff --git a/Engine/Plugins/Editor/PluginBrowser/Source/PluginBrowser/Private/SNewPluginWizard.cpp b/Engine/Plugins/Editor/PluginBrowser/Source/PluginBrowser/Private/SNewPluginWizard.cpp
--- a/Engine/Plugins/Editor/PluginBrowser/Source/PluginBrowser/Private/SNewPluginWizard.cpp
+++ b/Engine/Plugins/Editor/PluginBrowser/Source/PluginBrowser/Private/SNewPluginWizard.cpp
@@ -742,7 +742,7 @@ FReply SNewPluginWizard::OnCreatePluginClicked()
 	// If it contains code, we need the user to restart to enable it. Otherwise, we can just mount it now.
 	if (bSucceeded && bHasModules)
 	{
-		FString ProjectFileName = FPaths::GetProjectFilePath();
+		FString ProjectFileName = IFileManager::Get().ConvertToAbsolutePathForExternalAppForWrite(*FPaths::GetProjectFilePath());
 		FString Arguments = FString::Printf(TEXT("%s %s -TargetType=Editor -Plugin=\"%s\" -Project=\"%s\" -Progress -NoHotReloadFromIDE"), FModuleManager::Get().GetUBTConfiguration(), FPlatformMisc::GetUBTPlatform(), *UPluginFilePath, *ProjectFileName);
 		if (!FDesktopPlatformModule::Get()->RunUnrealBuildTool(LOCTEXT("Compiling", "Compiling..."), FPaths::RootDir(), Arguments, GWarn))
 		{
diff --git a/Engine/Source/Programs/UnrealBuildTool/Configuration/TargetDescriptor.cs b/Engine/Source/Programs/UnrealBuildTool/Configuration/TargetDescriptor.cs
--- a/Engine/Source/Programs/UnrealBuildTool/Configuration/TargetDescriptor.cs
+++ b/Engine/Source/Programs/UnrealBuildTool/Configuration/TargetDescriptor.cs
@@ -252,6 +252,12 @@ namespace UnrealBuildTool
 				throw new BuildException("No configurations specified for target");
 			}
 
+			// Make sure the project file exists
+			if(ProjectFile != null && !FileReference.Exists(ProjectFile))
+			{
+				throw new BuildException("Unable to find project '{0}'.", ProjectFile);
+			}
+
 			// Expand all the platforms, architectures and configurations
 			foreach(UnrealTargetPlatform Platform in Platforms)
 			{
