PR #6178: Filter by enabled state in plugin browser (Contributed by [REDACTED_PERSON])

diff --git a/Engine/Plugins/Editor/PluginBrowser/Source/PluginBrowser/Private/SPluginBrowser.cpp b/Engine/Plugins/Editor/PluginBrowser/Source/PluginBrowser/Private/SPluginBrowser.cpp
--- a/Engine/Plugins/Editor/PluginBrowser/Source/PluginBrowser/Private/SPluginBrowser.cpp
+++ b/Engine/Plugins/Editor/PluginBrowser/Source/PluginBrowser/Private/SPluginBrowser.cpp
@@ -21,6 +21,7 @@
 #include "Interfaces/IProjectManager.h"
 #include "ProjectDescriptor.h"
 #include "GameProjectGenerationModule.h"
+#include "Framework/MultiBox/MultiBoxBuilder.h"
 
 #define LOCTEXT_NAMESPACE "PluginsEditor"
 
@@ -95,6 +96,45 @@ void SPluginBrowser::Construct( const FArguments& Args )
 
 	PluginCategories = SNew( SPluginCategoryTree, SharedThis( this ) );
 
+	FMenuBuilder DetailViewOptions(true, nullptr);
+	DetailViewOptions.AddMenuEntry(
+		LOCTEXT("ShowOnlyEnabled", "Show Only Enabled"),
+		LOCTEXT("ShowOnlyEnabled_ToolTip", "Displays only the plugins which are enabled"),
+		FSlateIcon(),
+		FUIAction(
+			FExecuteAction::CreateLambda([this]
+			{
+				PluginCategories->ToggleFilterType(SPluginCategoryTree::EFilterType::OnlyEnabled);
+			}),
+			FCanExecuteAction(),
+			FIsActionChecked::CreateLambda([this]() -> bool
+			{
+				return PluginCategories->IsFilterEnabled(SPluginCategoryTree::EFilterType::OnlyEnabled);
+			})
+		),
+		NAME_None,
+		EUserInterfaceActionType::ToggleButton
+	);
+
+	DetailViewOptions.AddMenuEntry(
+		LOCTEXT("ShowOnlyDisabled", "Show Only Disabled"),
+		LOCTEXT("ShowOnlyDisabled_ToolTip", "Displays only the plugins which are disabled"),
+		FSlateIcon(),
+		FUIAction(
+			FExecuteAction::CreateLambda([this]
+			{
+				PluginCategories->ToggleFilterType(SPluginCategoryTree::EFilterType::OnlyDisabled);
+			}),
+			FCanExecuteAction(),
+			FIsActionChecked::CreateLambda([this]() -> bool
+			{
+				return PluginCategories->IsFilterEnabled(SPluginCategoryTree::EFilterType::OnlyDisabled);
+			})
+		),
+		NAME_None,
+		EUserInterfaceActionType::ToggleButton
+	);
+
 	TSharedRef<SVerticalBox> MainContent = SNew( SVerticalBox )
 	+SVerticalBox::Slot()
 	[
@@ -102,7 +142,12 @@ void SPluginBrowser::Construct( const FArguments& Args )
 		+SSplitter::Slot()
 		.Value(.3f)
 		[
-			PluginCategories.ToSharedRef()
+			SNew(SBorder)
+			.Padding(8.0f)
+			.BorderImage(FEditorStyle::GetBrush("ToolPanel.GroupBorder"))
+			[
+				PluginCategories.ToSharedRef()
+			]
 		]
 		+SSplitter::Slot()
 		[
@@ -129,6 +174,25 @@ void SPluginBrowser::Construct( const FArguments& Args )
 					SAssignNew( SearchBoxPtr, SSearchBox )
 					.OnTextChanged( this, &SPluginBrowser::SearchBox_OnPluginSearchTextChanged )
 				]
+
+				+ SHorizontalBox::Slot()
+				.AutoWidth()
+				[
+					SNew(SComboButton)
+					.ContentPadding(0)
+					.ForegroundColor(FSlateColor::UseForeground())
+					.ButtonStyle(FEditorStyle::Get(), "ToggleButton")
+					.AddMetaData<FTagMetaData>(FTagMetaData(TEXT("ViewOptions")))
+					.MenuContent()
+					[
+						DetailViewOptions.MakeWidget()
+					]
+					.ButtonContent()
+					[
+						SNew(SImage)
+						.Image(FEditorStyle::GetBrush("GenericViewButton"))
+					]
+				]
 			]
 
 			+SVerticalBox::Slot()
@@ -172,6 +236,8 @@ void SPluginBrowser::Construct( const FArguments& Args )
 						SNew(SButton)
 						.Text(LOCTEXT("PluginSettingsRestartEditor", "Restart Now"))
 						.OnClicked(this, &SPluginBrowser::HandleRestartEditorButtonClicked)
+						.TextStyle(FEditorStyle::Get(), "LargeText")
+						.ButtonStyle(FEditorStyle::Get(), "FlatButton.Default")
 					]
 				]
 			]
@@ -212,28 +278,28 @@ void SPluginBrowser::Construct( const FArguments& Args )
 						SNew(SButton)
 						.Text(LOCTEXT("PluginSettingsUpgradeToStudio", "Convert Project & Restart"))
 						.OnClicked(this, &SPluginBrowser::HandleUpgradeToUnrealStudioButtonClicked)
+						.TextStyle(FEditorStyle::Get(), "LargeText")
+						.ButtonStyle(FEditorStyle::Get(), "FlatButton.Default")
 					]
 				]
 			]
-		]
-	];
-
-	const FText NewPluginTooltip = LOCTEXT("NewPluginEnabled", "Click here to open the Plugin Creator dialog.");
 
-	MainContent->AddSlot()
-	.AutoHeight()
-	.Padding(5)
-	.HAlign(HAlign_Right)
-	[
-		SNew(SButton)
-		.ContentPadding(5)
-		.IsEnabled(true)
-		.ToolTip(SNew(SToolTip).Text(NewPluginTooltip))
-		.TextStyle(FEditorStyle::Get(), "LargeText")
-		.ButtonStyle(FEditorStyle::Get(), "FlatButton.Success")
-		.HAlign(HAlign_Center)
-		.Text(LOCTEXT("NewPluginLabel", "New Plugin"))
-		.OnClicked(this, &SPluginBrowser::HandleNewPluginButtonClicked)
+			+ SVerticalBox::Slot()
+			.AutoHeight()
+			.Padding(FMargin(PaddingAmount, PaddingAmount, PaddingAmount, 0.0f))
+			.HAlign(HAlign_Right)
+			[
+				SNew(SButton)
+				.ContentPadding(5)
+				.IsEnabled(true)
+				.ToolTip(SNew(SToolTip).Text(LOCTEXT("NewPluginEnabled", "Click here to open the Plugin Creator dialog.")))
+				.TextStyle(FEditorStyle::Get(), "LargeText")
+				.ButtonStyle(FEditorStyle::Get(), "FlatButton.Success")
+				.HAlign(HAlign_Center)
+				.Text(LOCTEXT("NewPluginLabel", "New Plugin"))
+				.OnClicked(this, &SPluginBrowser::HandleNewPluginButtonClicked)
+			]
+		]
 	];
 
 	ChildSlot
diff --git a/Engine/Plugins/Editor/PluginBrowser/Source/PluginBrowser/Private/SPluginCategoryTree.cpp b/Engine/Plugins/Editor/PluginBrowser/Source/PluginBrowser/Private/SPluginCategoryTree.cpp
--- a/Engine/Plugins/Editor/PluginBrowser/Source/PluginBrowser/Private/SPluginCategoryTree.cpp
+++ b/Engine/Plugins/Editor/PluginBrowser/Source/PluginBrowser/Private/SPluginCategoryTree.cpp
@@ -13,6 +13,8 @@ void SPluginCategoryTree::Construct( const FArguments& Args, const TSharedRef< S
 {
 	OwnerWeak = Owner;
 
+	FilterType = EFilterType::None;
+
 	// Create the root categories
 	BuiltInCategory = MakeShareable(new FPluginCategory(NULL, TEXT("Built-In"), LOCTEXT("BuiltInCategoryName", "Built-In")));
 	InstalledCategory = MakeShareable(new FPluginCategory(NULL, TEXT("Installed"), LOCTEXT("InstalledCategoryName", "Installed")));
@@ -87,6 +89,24 @@ void SPluginCategoryTree::RebuildAndFilterCategoryTree()
 			continue;
 		}
 
+		switch (FilterType)
+		{
+		case SPluginCategoryTree::EFilterType::None:
+			break;
+		case SPluginCategoryTree::EFilterType::OnlyEnabled:
+			if (!Plugin->IsEnabled())
+			{
+				continue;
+			}
+			break;
+		case SPluginCategoryTree::EFilterType::OnlyDisabled:
+			if (Plugin->IsEnabled())
+			{
+				continue;
+			}
+			break;
+		}
+
 		// Figure out which base category this plugin belongs in
 		TSharedPtr<FPluginCategory> RootCategory;
 		if (Plugin->GetType() == EPluginType::Mod)
@@ -272,6 +292,17 @@ void SPluginCategoryTree::SetNeedsRefresh()
 	RegisterActiveTimer (0.f, FWidgetActiveTimerDelegate::CreateSP (this, &SPluginCategoryTree::TriggerCategoriesRefresh));
 }
 
+bool SPluginCategoryTree::IsFilterEnabled(EFilterType FilterValue) const
+{
+	return (FilterType == FilterValue);
+}
+
+void SPluginCategoryTree::ToggleFilterType(EFilterType FilterValue)
+{
+	FilterType = IsFilterEnabled(FilterValue) ? EFilterType::None : FilterValue;
+	SetNeedsRefresh();
+}
+
 EActiveTimerReturnType SPluginCategoryTree::TriggerCategoriesRefresh(double InCurrentTime, float InDeltaTime)
 {
 	RebuildAndFilterCategoryTree();
diff --git a/Engine/Plugins/Editor/PluginBrowser/Source/PluginBrowser/Private/SPluginCategoryTree.h b/Engine/Plugins/Editor/PluginBrowser/Source/PluginBrowser/Private/SPluginCategoryTree.h
--- a/Engine/Plugins/Editor/PluginBrowser/Source/PluginBrowser/Private/SPluginCategoryTree.h
+++ b/Engine/Plugins/Editor/PluginBrowser/Source/PluginBrowser/Private/SPluginCategoryTree.h
@@ -17,9 +17,15 @@ class SPluginBrowser;
  */
 class SPluginCategoryTree : public SCompoundWidget
 {
-
 public:
 
+	enum class EFilterType : uint8
+	{
+		None,
+		OnlyEnabled,
+		OnlyDisabled,
+	};
+
 	SLATE_BEGIN_ARGS( SPluginCategoryTree )
 	{
 	}
@@ -48,6 +54,9 @@ public:
 	/** Signal that the categories list needs to be refreshed */
 	void SetNeedsRefresh();
 
+	void ToggleFilterType(EFilterType FilterValue);
+	bool IsFilterEnabled(EFilterType FilterValue) const;
+
 private:
 	
 	/** Called to generate a widget for the specified tree item */
@@ -87,5 +96,7 @@ private:
 
 	/** Category for mods */
 	TSharedPtr<FPluginCategory> ModCategory;
+
+	EFilterType FilterType;
 };
 
