Scan description text when searching for plugins in the plugin browser.

diff --git a/Engine/Plugins/Editor/PluginBrowser/Source/PluginBrowser/Private/SPluginBrowser.cpp b/Engine/Plugins/Editor/PluginBrowser/Source/PluginBrowser/Private/SPluginBrowser.cpp
--- a/Engine/Plugins/Editor/PluginBrowser/Source/PluginBrowser/Private/SPluginBrowser.cpp
+++ b/Engine/Plugins/Editor/PluginBrowser/Source/PluginBrowser/Private/SPluginBrowser.cpp
@@ -85,7 +85,9 @@ void SPluginBrowser::Construct( const FArguments& Args )
 		static void PluginToStringArray( const IPlugin* Plugin, OUT TArray< FString >& StringArray )
 		{
 			// NOTE: Only the friendly name is searchable for now.  We don't display the actual plugin name in the UI.
-			StringArray.Add( Plugin->GetDescriptor().FriendlyName );
+			const FPluginDescriptor& Descriptor = Plugin->GetDescriptor();
+			StringArray.Add( Descriptor.FriendlyName );
+			StringArray.Add( Descriptor.Description );
 		}
 	};
 
diff --git a/Engine/Plugins/Editor/PluginBrowser/Source/PluginBrowser/Private/SPluginTile.cpp b/Engine/Plugins/Editor/PluginBrowser/Source/PluginBrowser/Private/SPluginTile.cpp
--- a/Engine/Plugins/Editor/PluginBrowser/Source/PluginBrowser/Private/SPluginTile.cpp
+++ b/Engine/Plugins/Editor/PluginBrowser/Source/PluginBrowser/Private/SPluginTile.cpp
@@ -329,6 +329,7 @@ void SPluginTile::RecreateWidgets()
 											[
 												SNew(STextBlock)
 													.Text(FText::FromString(PluginDescriptor.Description))
+													.HighlightText_Raw(&OwnerWeak.Pin()->GetOwner().GetPluginTextFilter(), &FPluginTextFilter::GetRawFilterText)
 													.AutoWrapText(true)
 											]
 
