Prevent selection changing when changing plugin browser filter.

diff --git a/Engine/Plugins/Editor/PluginBrowser/Source/PluginBrowser/Private/SPluginCategory.h b/Engine/Plugins/Editor/PluginBrowser/Source/PluginBrowser/Private/SPluginCategory.h
--- a/Engine/Plugins/Editor/PluginBrowser/Source/PluginBrowser/Private/SPluginCategory.h
+++ b/Engine/Plugins/Editor/PluginBrowser/Source/PluginBrowser/Private/SPluginCategory.h
@@ -37,6 +37,19 @@ public:
 		  DisplayName(InDisplayName)
 	{
 	}
+
+	/** Finds a sub-category by name */
+	TSharedPtr<FPluginCategory> FindSubCategory(const FString& SubCategoryName)
+	{
+		for (TSharedPtr<FPluginCategory>& SubCategory : SubCategories)
+		{
+			if (SubCategory->Name == SubCategoryName)
+			{
+				return SubCategory;
+			}
+		}
+		return TSharedPtr<FPluginCategory>();
+	}
 };
 
 /**
diff --git a/Engine/Plugins/Editor/PluginBrowser/Source/PluginBrowser/Private/SPluginCategoryTree.cpp b/Engine/Plugins/Editor/PluginBrowser/Source/PluginBrowser/Private/SPluginCategoryTree.cpp
--- a/Engine/Plugins/Editor/PluginBrowser/Source/PluginBrowser/Private/SPluginCategoryTree.cpp
+++ b/Engine/Plugins/Editor/PluginBrowser/Source/PluginBrowser/Private/SPluginCategoryTree.cpp
@@ -66,22 +66,21 @@ static void ResetCategories(TArray<TSharedPtr<FPluginCategory>>& Categories)
 
 void SPluginCategoryTree::RebuildAndFilterCategoryTree()
 {
-	// Get a plugin from the currently selected category, so we can track it if it's removed
-	TSharedPtr<IPlugin> TrackPlugin = nullptr;
+	// Get the path to the first currently selected category
+	TArray<FString> SelectCategoryPath;
 	for(TSharedPtr<FPluginCategory> SelectedItem: TreeView->GetSelectedItems())
 	{
-		if(SelectedItem->Plugins.Num() > 0)
+		for (const FPluginCategory* Category = SelectedItem.Get(); Category != nullptr; Category = Category->ParentCategory.Pin().Get())
 		{
-			TrackPlugin = SelectedItem->Plugins[0];
-			break;
+			SelectCategoryPath.Insert(Category->Name, 0);
 		}
+		break;
 	}
 
 	// Clear the list of plugins in each current category
 	ResetCategories(RootCategories);
 
 	// Add all the known plugins into categories
-	TSharedPtr<FPluginCategory> SelectCategory;
 	for(TSharedRef<IPlugin> Plugin: IPluginManager::Get().GetDiscoveredPlugins())
 	{
 		if (Plugin->IsHidden())
@@ -160,12 +159,6 @@ void SPluginCategoryTree::RebuildAndFilterCategoryTree()
 			ParentCategory->Plugins.Add(Plugin);
 			ParentCategory = ParentCategory->ParentCategory.Pin();
 		}
-
-		// Update the selection if this is the plugin we're tracking
-		if(TrackPlugin == Plugin)
-		{
-			SelectCategory = FoundCategory;
-		}
 	}
 
 	// Remove any empty categories, keeping track of which items are still selected
@@ -180,6 +173,29 @@ void SPluginCategoryTree::RebuildAndFilterCategoryTree()
 		}
 	}
 
+	// Resolve the path to the category to select
+	TSharedPtr<FPluginCategory> SelectCategory;
+	if (SelectCategoryPath.Num() > 0)
+	{
+		for (TSharedPtr<FPluginCategory> RootCategory : RootCategories)
+		{
+			if (RootCategory->Name == SelectCategoryPath[0])
+			{
+				SelectCategory = RootCategory;
+				for (int Idx = 1; Idx < SelectCategoryPath.Num(); Idx++)
+				{
+					TSharedPtr<FPluginCategory> SubCategory = SelectCategory->FindSubCategory(SelectCategoryPath[Idx]);
+					if (!SubCategory.IsValid())
+					{
+						break;
+					}
+					SelectCategory = SubCategory;
+				}
+				break;
+			}
+		}
+	}
+
 	// Build the new list of root plugin categories
 	RootCategories.Reset();
 	if(ModCategory->SubCategories.Num() > 0 || ModCategory->Plugins.Num() > 0)
