Plugins: Add a flag to mark a plugin as experimental.

#rnx

diff --git a/Engine/Plugins/Editor/PluginBrowser/Source/PluginBrowser/Private/PluginStyle.cpp b/Engine/Plugins/Editor/PluginBrowser/Source/PluginBrowser/Private/PluginStyle.cpp
--- a/Engine/Plugins/Editor/PluginBrowser/Source/PluginBrowser/Private/PluginStyle.cpp
+++ b/Engine/Plugins/Editor/PluginBrowser/Source/PluginBrowser/Private/PluginStyle.cpp
@@ -151,9 +151,10 @@ void FPluginStyle::Initialize()
 			}
 
 			FTextBlockStyle BetaText = FTextBlockStyle( NormalText )
+				.SetFont(DEFAULT_FONT("Bold", 14))
 				.SetColorAndOpacity( FLinearColor( 0.9f, 0.9f, 0.9f ) );
 			{
-				BetaText.Font.Size = 14;
+				BetaText.Font.Size = 10;
 				StyleSet->Set( "PluginTile.BetaText", BetaText );
 			}
 
diff --git a/Engine/Plugins/Editor/PluginBrowser/Source/PluginBrowser/Private/SPluginTile.cpp b/Engine/Plugins/Editor/PluginBrowser/Source/PluginBrowser/Private/SPluginTile.cpp
--- a/Engine/Plugins/Editor/PluginBrowser/Source/PluginBrowser/Private/SPluginTile.cpp
+++ b/Engine/Plugins/Editor/PluginBrowser/Source/PluginBrowser/Private/SPluginTile.cpp
@@ -280,33 +280,42 @@ void SPluginTile::RecreateWidgets()
 											[
 												SNew(SHorizontalBox)
 
+												// beta version label
 												+ SHorizontalBox::Slot()
 													.AutoWidth()
-													.Padding(0.0f, 0.0f, 0.0f, 1.0f) // Lower padding to align font with version number base
+													.VAlign(VAlign_Bottom)
 													[
 														SNew(SHorizontalBox)
-
-														// beta version icon
+														.Visibility((PluginDescriptor.bIsBetaVersion || PluginDescriptor.bIsExperimentalVersion) ? EVisibility::Visible : EVisibility::Collapsed)
 														+ SHorizontalBox::Slot()
 															.AutoWidth()
 															.VAlign(VAlign_Bottom)
-															.Padding(0.0f, 0.0f, 3.0f, 2.0f)
+															.Padding(0.0f, 0.0f, 0.0f, 2.0f)
 															[
 																SNew(SImage)
-																	.Visibility(PluginDescriptor.bIsBetaVersion ? EVisibility::Visible : EVisibility::Collapsed)
 																	.Image(FPluginStyle::Get()->GetBrush("PluginTile.BetaWarning"))
 															]
-
-														// version label
 														+ SHorizontalBox::Slot()
 															.AutoWidth()
 															.VAlign(VAlign_Bottom)
+															.Padding(2.0f, 0.0f, 8.0f, 1.0f)
 															[
 																SNew(STextBlock)
-																	.Text(PluginDescriptor.bIsBetaVersion ? LOCTEXT("PluginBetaVersionLabel", "BETA Version ") : LOCTEXT("PluginVersionLabel", "Version "))
+																	.TextStyle(FPluginStyle::Get(), "PluginTile.BetaText")
+																	.Text(PluginDescriptor.bIsBetaVersion ? LOCTEXT("PluginBetaVersionText", "BETA") : LOCTEXT("PluginExperimentalVersionText", "EXPERIMENTAL"))
 															]
 													]
 
+												// version number
+												+ SHorizontalBox::Slot()
+													.AutoWidth()
+													.VAlign( VAlign_Bottom )
+													.Padding(0.0f, 0.0f, 0.0f, 1.0f) // Lower padding to align font with version number base
+													[
+														SNew(STextBlock)
+															.Text(LOCTEXT("PluginVersionLabel", "Version "))
+													]
+
 												+ SHorizontalBox::Slot()
 													.AutoWidth()
 													.VAlign( VAlign_Bottom )
diff --git a/Engine/Plugins/Experimental/BlastPlugin/BlastPlugin.uplugin b/Engine/Plugins/Experimental/BlastPlugin/BlastPlugin.uplugin
--- a/Engine/Plugins/Experimental/BlastPlugin/BlastPlugin.uplugin
+++ b/Engine/Plugins/Experimental/BlastPlugin/BlastPlugin.uplugin
@@ -2,7 +2,7 @@
 	"FileVersion": 3,
 	"Version": 1,
 	"VersionName": "0.1",
-	"FriendlyName": "Blast Plugin (Experimental)",
+	"FriendlyName": "Blast Plugin",
 	"Description": "A Plugin to expose Blast authoring functionality.",
 	"Category": "Editor",
 	"CreatedBy": "Epic Games, Inc.",
@@ -12,7 +12,7 @@
 	"SupportURL": "",
 	"EnabledByDefault": false,
 	"CanContainContent": false,
-	"IsBetaVersion": true,
+	"IsExperimentalVersion": true,
 	"Installed": false,
 	"Modules": [
 		{
diff --git a/Engine/Source/Programs/UnrealBuildTool/System/PluginDescriptor.cs b/Engine/Source/Programs/UnrealBuildTool/System/PluginDescriptor.cs
--- a/Engine/Source/Programs/UnrealBuildTool/System/PluginDescriptor.cs
+++ b/Engine/Source/Programs/UnrealBuildTool/System/PluginDescriptor.cs
@@ -154,6 +154,11 @@ namespace UnrealBuildTool
 		/// </summary>
 		public bool bIsBetaVersion;
 
+		/// <summary>
+		/// Marks the plugin as experimental in the UI
+		/// </summary>
+		public bool bIsExperimentalVersion;
+
 		/// <summary>
 		/// Set for plugins which are installed
 		/// </summary>
@@ -272,6 +277,7 @@ namespace UnrealBuildTool
 
 			RawObject.TryGetBoolField("CanContainContent", out bCanContainContent);
 			RawObject.TryGetBoolField("IsBetaVersion", out bIsBetaVersion);
+			RawObject.TryGetBoolField("IsExperimentalVersion", out bIsExperimentalVersion);
 			RawObject.TryGetBoolField("Installed", out bInstalled);
 
 			bool bCanBeUsedWithUnrealHeaderTool;
@@ -359,11 +365,15 @@ namespace UnrealBuildTool
 				Writer.WriteValue("EnabledByDefault", bEnabledByDefault.Value);
 			}
 			Writer.WriteValue("CanContainContent", bCanContainContent);
-			if(bIsBetaVersion)
+			if (bIsBetaVersion)
 			{
 				Writer.WriteValue("IsBetaVersion", bIsBetaVersion);
 			}
-			if(bInstalled)
+			if (bIsExperimentalVersion)
+			{
+				Writer.WriteValue("IsExperimentalVersion", bIsExperimentalVersion);
+			}
+			if (bInstalled)
 			{
 				Writer.WriteValue("Installed", bInstalled);
 			}
diff --git a/Engine/Source/Runtime/Projects/Private/PluginDescriptor.cpp b/Engine/Source/Runtime/Projects/Private/PluginDescriptor.cpp
--- a/Engine/Source/Runtime/Projects/Private/PluginDescriptor.cpp
+++ b/Engine/Source/Runtime/Projects/Private/PluginDescriptor.cpp
@@ -30,6 +30,7 @@ FPluginDescriptor::FPluginDescriptor()
 	, EnabledByDefault(EPluginEnabledByDefault::Unspecified)
 	, bCanContainContent(false)
 	, bIsBetaVersion(false)
+	, bIsExperimentalVersion(false)
 	, bInstalled(false)
 	, bRequiresBuildPlatform(false)
 	, bIsHidden(false)
@@ -137,6 +138,7 @@ bool FPluginDescriptor::Read(const FJsonObject& Object, FText& OutFailReason)
 
 	Object.TryGetBoolField(TEXT("CanContainContent"), bCanContainContent);
 	Object.TryGetBoolField(TEXT("IsBetaVersion"), bIsBetaVersion);
+	Object.TryGetBoolField(TEXT("IsExperimentalVersion"), bIsExperimentalVersion);
 	Object.TryGetBoolField(TEXT("Installed"), bInstalled);
 	Object.TryGetBoolField(TEXT("RequiresBuildPlatform"), bRequiresBuildPlatform);
 	Object.TryGetBoolField(TEXT("Hidden"), bIsHidden);
@@ -207,6 +209,7 @@ void FPluginDescriptor::Write(TJsonWriter<>& Writer) const
 	}
 	Writer.WriteValue(TEXT("CanContainContent"), bCanContainContent);
 	Writer.WriteValue(TEXT("IsBetaVersion"), bIsBetaVersion);
+	Writer.WriteValue(TEXT("IsExperimentalVersion"), bIsExperimentalVersion);
 	Writer.WriteValue(TEXT("Installed"), bInstalled);
 
 	if(SupportedTargetPlatforms.Num() > 0)
diff --git a/Engine/Source/Runtime/Projects/Public/PluginDescriptor.h b/Engine/Source/Runtime/Projects/Public/PluginDescriptor.h
--- a/Engine/Source/Runtime/Projects/Public/PluginDescriptor.h
+++ b/Engine/Source/Runtime/Projects/Public/PluginDescriptor.h
@@ -82,6 +82,9 @@ struct PROJECTS_API FPluginDescriptor
 	/** Marks the plugin as beta in the UI */
 	bool bIsBetaVersion;
 
+	/** Marks the plugin as experimental in the UI */
+	bool bIsExperimentalVersion;
+
 	/** Signifies that the plugin was installed on top of the engine */
 	bool bInstalled;
 
