Removed logic for Studio/Enterprise specific projects. Enterprise flag is still there in the project file for analytics.

#rnx

diff --git a/Engine/Plugins/Editor/PluginBrowser/Source/PluginBrowser/Private/SPluginBrowser.cpp b/Engine/Plugins/Editor/PluginBrowser/Source/PluginBrowser/Private/SPluginBrowser.cpp
--- a/Engine/Plugins/Editor/PluginBrowser/Source/PluginBrowser/Private/SPluginBrowser.cpp
+++ b/Engine/Plugins/Editor/PluginBrowser/Source/PluginBrowser/Private/SPluginBrowser.cpp
@@ -244,47 +244,6 @@ void SPluginBrowser::Construct( const FArguments& Args )
 				]
 			]
 
-			+SVerticalBox::Slot()
-			.AutoHeight()
-			.Padding(FMargin(PaddingAmount, PaddingAmount, PaddingAmount, 0.0f))
-			[
-				SNew(SBorder)
-				.BorderBackgroundColor(FLinearColor::Yellow)
-				.BorderImage(FEditorStyle::GetBrush("ToolPanel.GroupBorder"))
-				.Padding(8.0f)
-				.Visibility(this, &SPluginBrowser::HandleUpgradeToUnrealStudioVisibility)
-				[
-					SNew( SHorizontalBox )
-
-					+SHorizontalBox::Slot()
-					.AutoWidth()
-					.VAlign(VAlign_Center)
-					.Padding(FMargin(0.0f, 0.0f, 8.0f, 0.0f))
-					[
-						SNew(SImage)
-						.Image(FPluginStyle::Get()->GetBrush("Plugins.Warning"))
-					]
-
-					+SHorizontalBox::Slot()
-					.VAlign(VAlign_Center)
-					[
-						SNew(STextBlock)
-						.Text(LOCTEXT("PluginSettingsUpgradeToStudioNotice", "The current project is not an Unreal Studio project. Click 'Convert Project & Restart' to use Unreal Studio plugins."))
-					]
-
-					+SHorizontalBox::Slot()
-					.AutoWidth()
-					.VAlign(VAlign_Center)
-					.HAlign(HAlign_Right)
-					[
-						SNew(SButton)
-						.Text(LOCTEXT("PluginSettingsUpgradeToStudio", "Convert Project & Restart"))
-						.OnClicked(this, &SPluginBrowser::HandleUpgradeToUnrealStudioButtonClicked)
-						.TextStyle(FEditorStyle::Get(), "LargeText")
-						.ButtonStyle(FEditorStyle::Get(), "FlatButton.Default")
-					]
-				]
-			]
 
 			+ SVerticalBox::Slot()
 			.AutoHeight()
@@ -316,27 +275,6 @@ EVisibility SPluginBrowser::HandleRestartEditorNoticeVisibility() const
 	return FPluginBrowserModule::Get().HasPluginsPendingEnable() ? EVisibility::Visible : EVisibility::Collapsed;
 }
 
-EVisibility SPluginBrowser::HandleUpgradeToUnrealStudioVisibility() const
-{
-	const FProjectDescriptor* CurrentProject = IProjectManager::Get().GetCurrentProject();
-
-	TSharedPtr< FPluginCategory > SelectedCategory = GetSelectedCategory();
-
-	// Don't display the message for the root category or the project is an Unreal Studio project
-	if ( SelectedCategory.IsValid() && SelectedCategory->ParentCategory.IsValid() && !IProjectManager::Get().IsEnterpriseProject() ) 
-	{
-		for ( TSharedRef< IPlugin >& Plugin : SelectedCategory->Plugins )
-		{
-			if ( Plugin->GetType() == EPluginType::Enterprise )
-			{
-				return EVisibility::Visible;
-			}
-		}
-	}
-
-	return EVisibility::Collapsed;
-}
-
 FReply SPluginBrowser::HandleRestartEditorButtonClicked() const
 {
 	const bool bWarn = false;
@@ -344,24 +282,6 @@ FReply SPluginBrowser::HandleRestartEditorButtonClicked() const
 	return FReply::Handled();
 }
 
-FReply SPluginBrowser::HandleUpgradeToUnrealStudioButtonClicked() const
-{
-	if ( IProjectManager::Get().GetCurrentProject() )
-	{
-		FGameProjectGenerationModule::Get().TryMakeProjectFileWriteable(FPaths::GetProjectFilePath());
-
-		IProjectManager::Get().SetIsEnterpriseProject(true);
-
-		FText FailMessage;
-		IProjectManager::Get().SaveCurrentProjectToDisk(FailMessage);
-
-		const bool bWarn = false;
-		FUnrealEdMisc::Get().RestartEditor(bWarn);
-	}
-	
-	return FReply::Handled();
-}
-
 void SPluginBrowser::SearchBox_OnPluginSearchTextChanged( const FText& NewText )
 {
 	PluginTextFilter->SetRawFilterText( NewText );
diff --git a/Engine/Plugins/Editor/PluginBrowser/Source/PluginBrowser/Private/SPluginBrowser.h b/Engine/Plugins/Editor/PluginBrowser/Source/PluginBrowser/Private/SPluginBrowser.h
--- a/Engine/Plugins/Editor/PluginBrowser/Source/PluginBrowser/Private/SPluginBrowser.h
+++ b/Engine/Plugins/Editor/PluginBrowser/Source/PluginBrowser/Private/SPluginBrowser.h
@@ -61,15 +61,9 @@ private:
 	/** @return Is the "restart required" notice visible? */
 	EVisibility HandleRestartEditorNoticeVisibility() const;
 
-	/** @return Is the "update to Unreal Studio" notice visible? */
-	EVisibility HandleUpgradeToUnrealStudioVisibility() const;
-
 	/** Handle the "restart now" button being clicked */
 	FReply HandleRestartEditorButtonClicked() const;
 
-	/** Handle the "upgrade to Unreal Studio" button being clicked */
-	FReply HandleUpgradeToUnrealStudioButtonClicked() const;
-
 	/** Called when the text in the search box was changed */
 	void SearchBox_OnPluginSearchTextChanged( const FText& NewText );
 
diff --git a/Engine/Plugins/Editor/PluginBrowser/Source/PluginBrowser/Private/SPluginTile.cpp b/Engine/Plugins/Editor/PluginBrowser/Source/PluginBrowser/Private/SPluginTile.cpp
--- a/Engine/Plugins/Editor/PluginBrowser/Source/PluginBrowser/Private/SPluginTile.cpp
+++ b/Engine/Plugins/Editor/PluginBrowser/Source/PluginBrowser/Private/SPluginTile.cpp
@@ -55,12 +55,6 @@ void SPluginTile::RecreateWidgets()
 
 	const FPluginDescriptor& PluginDescriptor = Plugin->GetDescriptor();
 
-	// Disable Enterprise plugins if project is not an Enterprise project
-	if (Plugin->GetType() == EPluginType::Enterprise && !IProjectManager::Get().IsEnterpriseProject())
-	{
-		SetEnabled(false);
-	}
-
 	// Plugin thumbnail image
 	FString Icon128FilePath = Plugin->GetBaseDir() / TEXT("Resources/Icon128.png");
 	if(!FPlatformFileManager::Get().GetPlatformFile().FileExists(*Icon128FilePath))
diff --git a/Engine/Source/Editor/GameProjectGeneration/Private/GameProjectUtils.cpp b/Engine/Source/Editor/GameProjectGeneration/Private/GameProjectUtils.cpp
--- a/Engine/Source/Editor/GameProjectGeneration/Private/GameProjectUtils.cpp
+++ b/Engine/Source/Editor/GameProjectGeneration/Private/GameProjectUtils.cpp
@@ -1807,8 +1807,6 @@ bool GameProjectUtils::CreateProjectFromTemplate(const FProjectInformation& InPr
 		Project.EngineAssociation.Empty();
 		Project.EpicSampleNameHash = 0;
 
-		Project.bIsEnterpriseProject = InProjectInfo.bIsEnterpriseProject; // Force the enterprise flag to the value that was requested in the ProjectInfo.
-
 		// Fix up module names
 		const FString BaseSourceName = FPaths::GetBaseFilename(InProjectInfo.TemplateFile);
 		const FString BaseNewName = FPaths::GetBaseFilename(InProjectInfo.ProjectFilename);
diff --git a/Engine/Source/Editor/LevelEditor/Private/LevelEditor.cpp b/Engine/Source/Editor/LevelEditor/Private/LevelEditor.cpp
--- a/Engine/Source/Editor/LevelEditor/Private/LevelEditor.cpp
+++ b/Engine/Source/Editor/LevelEditor/Private/LevelEditor.cpp
@@ -100,10 +100,6 @@ public:
 		// Create the tooltip showing more detailed information
 		FFormatNamedArguments TooltipArgs;
 		FString TooltipVersionStr = EngineVersionString;
-		if (IProjectManager::Get().GetCurrentProject() && IProjectManager::Get().GetCurrentProject()->bIsEnterpriseProject)
-		{
-			TooltipVersionStr += TEXT(" Unreal Studio");
-		}
 		TooltipArgs.Add(TEXT("Version"), FText::FromString(TooltipVersionStr));
 		TooltipArgs.Add(TEXT("Branch"), FText::FromString(FEngineVersion::Current().GetBranch()));
 		TooltipArgs.Add(TEXT("BuildConfiguration"), EBuildConfigurations::ToText(BuildConfig));
diff --git a/Engine/Source/Runtime/Core/Private/Misc/App.cpp b/Engine/Source/Runtime/Core/Private/Misc/App.cpp
--- a/Engine/Source/Runtime/Core/Private/Misc/App.cpp
+++ b/Engine/Source/Runtime/Core/Private/Misc/App.cpp
@@ -229,35 +229,6 @@ bool FApp::IsEngineInstalled()
 	return EngineInstalledState == 1;
 }
 
-bool FApp::IsEnterpriseInstalled()
-{
-	static int32 EnterpriseInstalledState = -1;
-
-	if (EnterpriseInstalledState == -1)
-	{
-		bool bIsInstalledEnterprise = false;
-
-#if PLATFORM_DESKTOP
-		FString InstalledBuildFile = FPaths::RootDir() / TEXT("Enterprise/Build/InstalledBuild.txt");
-		FPaths::NormalizeFilename(InstalledBuildFile);
-		bIsInstalledEnterprise |= IFileManager::Get().FileExists(*InstalledBuildFile);
-#endif
-
-		// Allow commandline options to disable/enable installed engine behavior
-		if (bIsInstalledEnterprise)
-		{
-			bIsInstalledEnterprise = !FParse::Param(FCommandLine::Get(), TEXT("NotInstalledEnterprise"));
-		}
-		else
-		{
-			bIsInstalledEnterprise = FParse::Param(FCommandLine::Get(), TEXT("InstalledEnterprise"));
-		}
-		EnterpriseInstalledState = bIsInstalledEnterprise ? 1 : 0;
-	}
-
-	return EnterpriseInstalledState == 1;
-}
-
 #if PLATFORM_WINDOWS && defined(__clang__)
 bool FApp::IsUnattended() // @todo clang: Workaround for missing symbol export
 {
diff --git a/Engine/Source/Runtime/Core/Public/Misc/App.h b/Engine/Source/Runtime/Core/Public/Misc/App.h
--- a/Engine/Source/Runtime/Core/Public/Misc/App.h
+++ b/Engine/Source/Runtime/Core/Public/Misc/App.h
@@ -403,13 +403,6 @@ public:
 	 */
 	static bool IsEngineInstalled();
 
-	/**
-	 * Checks whether the enterprise components of this application have been installed.
-	 *
-	 * @return true if the enterprise components are installed, false otherwise.
-	 */
-	static bool IsEnterpriseInstalled();
-
 	/**
 	 * Checks whether this application runs unattended.
 	 *
