[PluginBrowser]
- Add "All" category that includes all plugins. Useful to search all plugins and not restrict the search to a given category or sub-category.
- Bonus: Display, sort and search by plugin name if friendly name is not set.

diff --git a/Engine/Plugins/Editor/PluginBrowser/Source/PluginBrowser/Private/SPluginBrowser.cpp b/Engine/Plugins/Editor/PluginBrowser/Source/PluginBrowser/Private/SPluginBrowser.cpp
--- a/Engine/Plugins/Editor/PluginBrowser/Source/PluginBrowser/Private/SPluginBrowser.cpp
+++ b/Engine/Plugins/Editor/PluginBrowser/Source/PluginBrowser/Private/SPluginBrowser.cpp
@@ -14,6 +14,7 @@
 #include "PluginStyle.h"
 #include "Widgets/Navigation/SBreadcrumbTrail.h"
 #include "SPluginCategoryTree.h"
+#include "SPluginTile.h"
 #include "Widgets/Input/SSearchBox.h"
 #include "PluginBrowserModule.h"
 #include "IDirectoryWatcher.h"
@@ -84,10 +85,9 @@ void SPluginBrowser::Construct( const FArguments& Args )
 	{
 		static void PluginToStringArray( const IPlugin* Plugin, OUT TArray< FString >& StringArray )
 		{
-			// NOTE: Only the friendly name is searchable for now.  We don't display the actual plugin name in the UI.
 			const FPluginDescriptor& Descriptor = Plugin->GetDescriptor();
-			StringArray.Add( Descriptor.FriendlyName );
-			StringArray.Add( Descriptor.Description );
+			StringArray.Add(SPluginTile::GetPluginDisplayName(Plugin));
+			StringArray.Add(Descriptor.Description);
 		}
 	};
 
@@ -291,7 +291,7 @@ void SPluginBrowser::SearchBox_OnPluginSearchTextChanged( const FText& NewText )
 
 TSharedPtr< FPluginCategory > SPluginBrowser::GetSelectedCategory() const
 {
-	return PluginCategories->GetSelectedCategory();
+	return PluginCategories.IsValid() ? PluginCategories->GetSelectedCategory() : nullptr;
 }
 
 
diff --git a/Engine/Plugins/Editor/PluginBrowser/Source/PluginBrowser/Private/SPluginCategoryTree.cpp b/Engine/Plugins/Editor/PluginBrowser/Source/PluginBrowser/Private/SPluginCategoryTree.cpp
--- a/Engine/Plugins/Editor/PluginBrowser/Source/PluginBrowser/Private/SPluginCategoryTree.cpp
+++ b/Engine/Plugins/Editor/PluginBrowser/Source/PluginBrowser/Private/SPluginCategoryTree.cpp
@@ -16,6 +16,7 @@ void SPluginCategoryTree::Construct( const FArguments& Args, const TSharedRef< S
 	FilterType = EFilterType::None;
 
 	// Create the root categories
+	AllCategory = MakeShareable(new FPluginCategory(NULL, TEXT("All"), LOCTEXT("AllCategoryName", "All")));
 	BuiltInCategory = MakeShareable(new FPluginCategory(NULL, TEXT("Built-In"), LOCTEXT("BuiltInCategoryName", "Built-In")));
 	InstalledCategory = MakeShareable(new FPluginCategory(NULL, TEXT("Installed"), LOCTEXT("InstalledCategoryName", "Installed")));
 	ProjectCategory = MakeShareable(new FPluginCategory(NULL, TEXT("Project"), LOCTEXT("ProjectCategoryName", "Project")));
@@ -159,6 +160,12 @@ void SPluginCategoryTree::RebuildAndFilterCategoryTree()
 			ParentCategory->Plugins.Add(Plugin);
 			ParentCategory = ParentCategory->ParentCategory.Pin();
 		}
+
+		// Add the plugin in the "All" category
+		if (AllCategory.IsValid())
+		{
+			AllCategory->Plugins.Add(Plugin);
+		}
 	}
 
 	// Remove any empty categories, keeping track of which items are still selected
@@ -214,6 +221,10 @@ void SPluginCategoryTree::RebuildAndFilterCategoryTree()
 	{
 		RootCategories.Add(ProjectCategory);
 	}
+	if (RootCategories.Num() > 0)
+	{
+		RootCategories.Insert(AllCategory, 0);
+	}
 
 	// Sort every single category alphabetically
 	for(TSharedPtr<FPluginCategory> RootCategory: RootCategories)
@@ -231,21 +242,25 @@ void SPluginCategoryTree::RebuildAndFilterCategoryTree()
 	TreeView->RequestTreeRefresh();
 
 	// Make sure we have something selected
-	if(RootCategories.Num() > 0)
+	if (RootCategories.Num() > 0)
 	{
-		if(SelectCategory.IsValid())
+		if (SelectCategory.IsValid())
 		{
 			TreeView->SetSelection(SelectCategory);
 		}
-		else if(RootCategories.Contains(ModCategory))
+		else if (RootCategories.Contains(ModCategory))
 		{
 			TreeView->SetSelection(ModCategory);
 		}
-		else if(RootCategories.Contains(InstalledCategory))
+		else if (RootCategories.Contains(InstalledCategory))
 		{
 			TreeView->SetSelection(InstalledCategory);
 		}
-		else if(RootCategories.Num() > 0)
+		else if (RootCategories.Contains(AllCategory))
+		{
+			TreeView->SetSelection(AllCategory);
+		}
+		else if (RootCategories.Num() > 0)
 		{
 			TreeView->SetSelection(RootCategories[0]);
 		}
diff --git a/Engine/Plugins/Editor/PluginBrowser/Source/PluginBrowser/Private/SPluginCategoryTree.h b/Engine/Plugins/Editor/PluginBrowser/Source/PluginBrowser/Private/SPluginCategoryTree.h
--- a/Engine/Plugins/Editor/PluginBrowser/Source/PluginBrowser/Private/SPluginCategoryTree.h
+++ b/Engine/Plugins/Editor/PluginBrowser/Source/PluginBrowser/Private/SPluginCategoryTree.h
@@ -85,6 +85,9 @@ private:
 	/** Root list of categories */
 	TArray<TSharedPtr<FPluginCategory>> RootCategories;
 
+	/** Category that includes all plugins */
+	TSharedPtr<FPluginCategory> AllCategory;
+
 	/** Category for built-in plugins */
 	TSharedPtr<FPluginCategory> BuiltInCategory;
 
diff --git a/Engine/Plugins/Editor/PluginBrowser/Source/PluginBrowser/Private/SPluginTile.cpp b/Engine/Plugins/Editor/PluginBrowser/Source/PluginBrowser/Private/SPluginTile.cpp
--- a/Engine/Plugins/Editor/PluginBrowser/Source/PluginBrowser/Private/SPluginTile.cpp
+++ b/Engine/Plugins/Editor/PluginBrowser/Source/PluginBrowser/Private/SPluginTile.cpp
@@ -41,6 +41,15 @@ void SPluginTile::Construct( const FArguments& Args, const TSharedRef<SPluginTil
 
 	RecreateWidgets();
 }
+FString SPluginTile::GetPluginDisplayName(const IPlugin* InPlugin)
+{
+	return InPlugin ? (!InPlugin->GetDescriptor().FriendlyName.IsEmpty() ? InPlugin->GetDescriptor().FriendlyName : InPlugin->GetName()) : "";
+}
+
+FText SPluginTile::GetPluginDisplayText() const
+{	
+	return FText::FromString(GetPluginDisplayName(Plugin.Get()));
+}
 
 void SPluginTile::RecreateWidgets()
 {
@@ -236,7 +245,7 @@ void SPluginTile::RecreateWidgets()
 											.Padding(PaddingAmount)
 											[
 												SNew(STextBlock)
-													.Text(FText::FromString(PluginDescriptor.FriendlyName))
+													.Text(GetPluginDisplayText())
 													.HighlightText_Raw(&OwnerWeak.Pin()->GetOwner().GetPluginTextFilter(), &FPluginTextFilter::GetRawFilterText)
 													.TextStyle(FPluginStyle::Get(), "PluginTile.NameText")
 											]
@@ -461,7 +470,7 @@ void SPluginTile::OnEnablePluginCheckboxChanged(ECheckBoxState NewCheckedState)
 		// If this is plugin is marked as beta, make sure the user is aware before enabling it.
 		if (PluginDescriptor.bIsBetaVersion)
 		{
-			FText WarningMessage = FText::Format(LOCTEXT("Warning_EnablingBetaPlugin", "Plugin '{0}' is a beta version and might be unstable or removed without notice. Please use with caution. Are you sure you want to enable the plugin?"), FText::FromString(PluginDescriptor.FriendlyName));
+			FText WarningMessage = FText::Format(LOCTEXT("Warning_EnablingBetaPlugin", "Plugin '{0}' is a beta version and might be unstable or removed without notice. Please use with caution. Are you sure you want to enable the plugin?"), GetPluginDisplayText());
 			if (EAppReturnType::No == FMessageDialog::Open(EAppMsgType::YesNo, WarningMessage))
 			{
 				return;
diff --git a/Engine/Plugins/Editor/PluginBrowser/Source/PluginBrowser/Private/SPluginTile.h b/Engine/Plugins/Editor/PluginBrowser/Source/PluginBrowser/Private/SPluginTile.h
--- a/Engine/Plugins/Editor/PluginBrowser/Source/PluginBrowser/Private/SPluginTile.h
+++ b/Engine/Plugins/Editor/PluginBrowser/Source/PluginBrowser/Private/SPluginTile.h
@@ -29,8 +29,14 @@ class SPluginTile : public SCompoundWidget
 	/** Widget constructor */
 	void Construct( const FArguments& Args, const TSharedRef< class SPluginTileList > Owner, TSharedRef<IPlugin> Plugin );
 
+	/** Returns plugin friendly name if not empty otherwise the plugin name. */
+	static FString GetPluginDisplayName(const IPlugin* InPlugin);
+
 private:
 
+	/** Returns text to display for the plugin name. */
+	FText GetPluginDisplayText() const;
+
 	/** Updates the contents of this tile */
 	void RecreateWidgets();
 
diff --git a/Engine/Plugins/Editor/PluginBrowser/Source/PluginBrowser/Private/SPluginTileList.cpp b/Engine/Plugins/Editor/PluginBrowser/Source/PluginBrowser/Private/SPluginTileList.cpp
--- a/Engine/Plugins/Editor/PluginBrowser/Source/PluginBrowser/Private/SPluginTileList.cpp
+++ b/Engine/Plugins/Editor/PluginBrowser/Source/PluginBrowser/Private/SPluginTileList.cpp
@@ -86,9 +86,9 @@ void SPluginTileList::RebuildAndFilterPluginList()
 		{
 			struct FPluginListItemSorter
 			{
-				bool operator()( const TSharedRef<IPlugin>& A, const TSharedRef<IPlugin>& B ) const
+				bool operator()(const TSharedRef<IPlugin>& A, const TSharedRef<IPlugin>& B) const
 				{
-					return A->GetDescriptor().FriendlyName < B->GetDescriptor().FriendlyName;
+					return SPluginTile::GetPluginDisplayName(&*A) < SPluginTile::GetPluginDisplayName(&*B);
 				}
 			};
 			PluginListItems.Sort( FPluginListItemSorter() );
