Fix the "New Plugins Available" prompt not dismissing immediately when clicking the dismiss button

diff --git a/Engine/Plugins/Editor/PluginBrowser/Source/PluginBrowser/Private/PluginBrowserModule.cpp b/Engine/Plugins/Editor/PluginBrowser/Source/PluginBrowser/Private/PluginBrowserModule.cpp
--- a/Engine/Plugins/Editor/PluginBrowser/Source/PluginBrowser/Private/PluginBrowserModule.cpp
+++ b/Engine/Plugins/Editor/PluginBrowser/Source/PluginBrowser/Private/PluginBrowserModule.cpp
@@ -162,13 +162,14 @@ TSharedRef<SDockTab> FPluginBrowserModule::SpawnPluginCreatorTab(const FSpawnTab
 void FPluginBrowserModule::OnMainFrameLoaded(TSharedPtr<SWindow> InRootWindow, bool bIsNewProjectWindow)
 {
 	// Show a popup notification that allows the user to enable any new plugins
-	if(!bIsNewProjectWindow && NewlyInstalledPlugins.Num() > 0 && !PluginBrowserTab.IsValid())
+	//if(!bIsNewProjectWindow && NewlyInstalledPlugins.Num() > 0 && !PluginBrowserTab.IsValid())
 	{
 		FNotificationInfo Info(LOCTEXT("NewPluginsPopupTitle", "New plugins are available"));
 		Info.bFireAndForget = false;
 		Info.bUseLargeFont = true;
 		Info.bUseThrobber = false;
-		Info.FadeOutDuration = 0.5f;
+		Info.FadeOutDuration = 0.0f;
+		Info.ExpireDuration = 0.0f;
 		Info.ButtonDetails.Add(FNotificationButtonInfo(LOCTEXT("NewPluginsPopupSettings", "Manage Plugins..."), LOCTEXT("NewPluginsPopupSettingsTT", "Open the plugin browser to enable plugins"), FSimpleDelegate::CreateRaw(this, &FPluginBrowserModule::OnNewPluginsPopupSettingsClicked)));
 		Info.ButtonDetails.Add(FNotificationButtonInfo(LOCTEXT("NewPluginsPopupDismiss", "Dismiss"), LOCTEXT("NewPluginsPopupDismissTT", "Dismiss this notification"), FSimpleDelegate::CreateRaw(this, &FPluginBrowserModule::OnNewPluginsPopupDismissClicked)));
 
