Missing edit to last file

diff --git a/Engine/Plugins/Editor/PluginBrowser/Source/PluginBrowser/Private/PluginBrowserModule.cpp b/Engine/Plugins/Editor/PluginBrowser/Source/PluginBrowser/Private/PluginBrowserModule.cpp
--- a/Engine/Plugins/Editor/PluginBrowser/Source/PluginBrowser/Private/PluginBrowserModule.cpp
+++ b/Engine/Plugins/Editor/PluginBrowser/Source/PluginBrowser/Private/PluginBrowserModule.cpp
@@ -162,7 +162,7 @@ TSharedRef<SDockTab> FPluginBrowserModule::SpawnPluginCreatorTab(const FSpawnTab
 void FPluginBrowserModule::OnMainFrameLoaded(TSharedPtr<SWindow> InRootWindow, bool bIsNewProjectWindow)
 {
 	// Show a popup notification that allows the user to enable any new plugins
-	//if(!bIsNewProjectWindow && NewlyInstalledPlugins.Num() > 0 && !PluginBrowserTab.IsValid())
+	if(!bIsNewProjectWindow && NewlyInstalledPlugins.Num() > 0 && !PluginBrowserTab.IsValid())
 	{
 		FNotificationInfo Info(LOCTEXT("NewPluginsPopupTitle", "New plugins are available"));
 		Info.bFireAndForget = false;
