Minor reskin tweaks to plugin creation UI

diff --git a/Engine/Plugins/Editor/PluginBrowser/Source/PluginBrowser/Private/PluginStyle.cpp b/Engine/Plugins/Editor/PluginBrowser/Source/PluginBrowser/Private/PluginStyle.cpp
--- a/Engine/Plugins/Editor/PluginBrowser/Source/PluginBrowser/Private/PluginStyle.cpp
+++ b/Engine/Plugins/Editor/PluginBrowser/Source/PluginBrowser/Private/PluginStyle.cpp
@@ -6,14 +6,14 @@
 #include "Styling/CoreStyle.h"
 #include "EditorStyleSet.h"
 #include "Interfaces/IPluginManager.h"
+#include "Styling/StarshipCoreStyle.h"
 
-#define IMAGE_BRUSH( RelativePath, ... ) FSlateImageBrush( FPluginStyle::InContent( RelativePath, ".png" ), __VA_ARGS__ )
-#define IMAGE_BRUSH_SVG( RelativePath, ... ) FSlateVectorImageBrush( FPluginStyle::InContent( RelativePath, ".svg" ), __VA_ARGS__ )
-#define BOX_BRUSH( RelativePath, ... ) FSlateBoxBrush( FPluginStyle::InContent( RelativePath, ".png" ), __VA_ARGS__ )
-#define BORDER_BRUSH( RelativePath, ... ) FSlateBorderBrush( FPluginStyle::InContent( RelativePath, ".png" ), __VA_ARGS__ )
-#define DEFAULT_FONT(...) FCoreStyle::GetDefaultFontStyle(__VA_ARGS__)
+#define RootToContentDir FPluginStyle::InContent
 
-FString FPluginStyle::InContent( const FString& RelativePath, const ANSICHAR* Extension )
+#include "Styling/SlateStyleMacros.h"
+#include "Styling/StyleColors.h"
+
+FString FPluginStyle::InContent( const FString& RelativePath, const TCHAR* Extension )
 {
 	static FString ContentDir = IPluginManager::Get().FindPlugin(TEXT("PluginBrowser"))->GetContentDir();
 	return ( ContentDir / RelativePath ) + Extension;
@@ -193,20 +193,19 @@ void FPluginStyle::Initialize()
 
 	// Plugin Creator
 	{
-		const float PaddingAmount = 5.0f;
-		StyleSet->Set( "PluginCreator.Padding", PaddingAmount);
-
-		StyleSet->Set( "PluginCreator.Background", new BOX_BRUSH( "AppTabContentArea", FMargin(4/16.0f) ) );
+		const FButtonStyle& BaseButtonStyle = FStarshipCoreStyle::GetCoreStyle().GetWidgetStyle<FButtonStyle>("Button");
+		StyleSet->Set("PluginPath.BrowseButton",
+			FButtonStyle(BaseButtonStyle)
+			.SetNormal(FSlateRoundedBoxBrush(FStyleColors::Secondary, 4.0f, FStyleColors::Secondary, 2.0f))
+			.SetHovered(FSlateRoundedBoxBrush(FStyleColors::Hover, 4.0f, FStyleColors::Hover, 2.0f))
+			.SetPressed(FSlateRoundedBoxBrush(FStyleColors::Header, 4.0f, FStyleColors::Header, 2.0f))
+			.SetNormalPadding(FMargin(2, 2, 2, 2))
+			.SetPressedPadding(FMargin(2, 3, 2, 1)));
 	}
 
 	FSlateStyleRegistry::RegisterSlateStyle( *StyleSet.Get() );
 };
 
-#undef IMAGE_BRUSH
-#undef BOX_BRUSH
-#undef BORDER_BRUSH
-#undef DEFAULT_FONT
-
 void FPluginStyle::Shutdown()
 {
 	if( StyleSet.IsValid() )
@@ -216,3 +215,5 @@ void FPluginStyle::Shutdown()
 		StyleSet.Reset();
 	}
 }
+
+#undef RootToContentDir
\ No newline at end of file
diff --git a/Engine/Plugins/Editor/PluginBrowser/Source/PluginBrowser/Private/PluginStyle.h b/Engine/Plugins/Editor/PluginBrowser/Source/PluginBrowser/Private/PluginStyle.h
--- a/Engine/Plugins/Editor/PluginBrowser/Source/PluginBrowser/Private/PluginStyle.h
+++ b/Engine/Plugins/Editor/PluginBrowser/Source/PluginBrowser/Private/PluginStyle.h
@@ -17,7 +17,7 @@ public:
 
 private:
 
-	static FString InContent( const FString& RelativePath, const ANSICHAR* Extension );
+	static FString InContent( const FString& RelativePath, const TCHAR* Extension );
 
 private:
 
diff --git a/Engine/Plugins/Editor/PluginBrowser/Source/PluginBrowser/Private/SFilePathBlock.cpp b/Engine/Plugins/Editor/PluginBrowser/Source/PluginBrowser/Private/SFilePathBlock.cpp
--- a/Engine/Plugins/Editor/PluginBrowser/Source/PluginBrowser/Private/SFilePathBlock.cpp
+++ b/Engine/Plugins/Editor/PluginBrowser/Source/PluginBrowser/Private/SFilePathBlock.cpp
@@ -8,6 +8,8 @@
 #include "Widgets/Layout/SGridPanel.h"
 #include "Widgets/Input/SEditableTextBox.h"
 #include "Widgets/Input/SButton.h"
+#include "Widgets/SBoxPanel.h"
+#include "PluginStyle.h"
 
 #define LOCTEXT_NAMESPACE "FilePathBlock"
 
@@ -28,116 +30,55 @@ void SFilePathBlock::Construct(const FArguments& InArgs)
 
 	ChildSlot
 	[
-		SNew(SGridPanel)
-		.FillColumn(0, 2.f)
-		.FillColumn(1, 1.f)
-
+		SNew(SHorizontalBox)
 		// Folder input
-		+ SGridPanel::Slot(0, 0)
+		+SHorizontalBox::Slot()
 		[
-			SNew(SOverlay)
-
-			+ SOverlay::Slot()
+			SNew(SBox)
+			.MinDesiredWidth(250.f)
 			[
-				SAssignNew(FolderPathTextBox, SEditableTextBox)
-				.Text(InArgs._FolderPath)
-				// Large right hand padding to make room for the browse button
-				.Padding(FMargin(5.f, 3.f, 25.f, 3.f))
-				.OnTextChanged(InArgs._OnFolderChanged)
-				.OnTextCommitted(InArgs._OnFolderCommitted)
-				.IsReadOnly(bReadOnlyFolderPath)
-			]
-
-			+ SOverlay::Slot()
-			.HAlign(HAlign_Right)
-			[
-				SNew(SButton)
-				.ButtonStyle(FEditorStyle::Get(), "FilePath.FolderButton")
-				.ContentPadding(FMargin(4.f, 0.f))
-				.OnClicked(InArgs._OnBrowseForFolder)
-				.ToolTipText(BrowseForFolderToolTipText)
-				.Text(LOCTEXT("...", "..."))
-				.IsEnabled( !bReadOnlyFolderPath )
-			]
-		]
-
-		// Folder label
-		+ SGridPanel::Slot(0, 1)
-		[
-			SNew(SOverlay)
-
-			+ SOverlay::Slot()
-			.VAlign(VAlign_Center)
-			[
-				SNew(SBox)
-				.HeightOverride(3)
+				SNew(SOverlay)
+				+ SOverlay::Slot()
 				[
-					SNew(SBorder)
-					.BorderImage(FEditorStyle::GetBrush("FilePath.GroupIndicator"))
-					.BorderBackgroundColor(FLinearColor(1.f, 1.f, 1.f, 0.5f))
-					.Padding(FMargin(150.f, 0.f))
+					SAssignNew(FolderPathTextBox, SEditableTextBox)
+					.Text(InArgs._FolderPath)
+					// Large right hand padding to make room for the browse button
+					.Padding(FMargin(5.f, 3.f, 25.f, 3.f))
+					.OnTextChanged(InArgs._OnFolderChanged)
+					.OnTextCommitted(InArgs._OnFolderCommitted)
+					.IsReadOnly(bReadOnlyFolderPath)
+					.HintText(LOCTEXT("Folder", "Folder"))
+					.ForegroundColor(FSlateColor::UseForeground())
 				]
-			]
 
-			+ SOverlay::Slot()
-			.HAlign(HAlign_Center)
-			.VAlign(VAlign_Center)
-			[
-				SNew(SBorder)
-				.Padding(5.f)
-				.BorderImage(InArgs._LabelBackgroundBrush)
-				.BorderBackgroundColor(InArgs._LabelBackgroundColor)
+				+ SOverlay::Slot()
+				.HAlign(HAlign_Right)
+				.Padding(1)
 				[
-					SNew(STextBlock)
-					.Text(LOCTEXT("Folder", "Folder"))
+					SNew(SButton)
+					.ButtonStyle(FPluginStyle::Get(), "PluginPath.BrowseButton")
+					.OnClicked(InArgs._OnBrowseForFolder)
+					.ContentPadding(0)
+					.ToolTipText(BrowseForFolderToolTipText)
+					.Text(LOCTEXT("...", "..."))
+					.IsEnabled(!bReadOnlyFolderPath)
 				]
 			]
 		]
-
 		// Name input
-		+ SGridPanel::Slot(1, 0)
+		+ SHorizontalBox::Slot()
 		.Padding(FMargin(5.f, 0.f, 0.f, 0.f))
 		.VAlign(VAlign_Center)
 		[
-			SAssignNew(NameTextBox, SEditableTextBox)
-			.Text(InArgs._Name)
-			.Padding(FMargin(5.f, 3.f))
-			.HintText(InArgs._NameHint)
-			.OnTextChanged(InArgs._OnNameChanged)
-			.OnTextCommitted(InArgs._OnNameCommitted)
-		]
-
-		// Name label
-		+ SGridPanel::Slot(1, 1)
-		.Padding(FMargin(5.f, 0.f, 0.f, 0.f))
-		[
-			SNew(SOverlay)
-
-			+ SOverlay::Slot()
-			.VAlign(VAlign_Center)
+			SNew(SBox)
+			.MinDesiredWidth(200.f)
 			[
-				SNew(SBox)
-				.HeightOverride(3)
-				[
-					SNew(SBorder)
-					.BorderImage(FEditorStyle::GetBrush("FilePath.GroupIndicator"))
-					.BorderBackgroundColor(FLinearColor(1.f, 1.f, 1.f, 0.5f))
-					.Padding(FMargin(75.f, 0.f))
-				]
-			]
-
-			+ SOverlay::Slot()
-			.HAlign(HAlign_Center)
-			.VAlign(VAlign_Center)
-			[
-				SNew(SBorder)
-				.Padding(5.f)
-				.BorderImage(InArgs._LabelBackgroundBrush)
-				.BorderBackgroundColor(InArgs._LabelBackgroundColor)
-				[
-					SNew(STextBlock)
-					.Text(LOCTEXT("Name", "Name"))
-				]
+				SAssignNew(NameTextBox, SEditableTextBox)
+				.Text(InArgs._Name)
+				.Padding(FMargin(5.f, 3.f))
+				.HintText(InArgs._NameHint)
+				.OnTextChanged(InArgs._OnNameChanged)
+				.OnTextCommitted(InArgs._OnNameCommitted)
 			]
 		]
 	];
diff --git a/Engine/Plugins/Editor/PluginBrowser/Source/PluginBrowser/Private/SFilePathBlock.h b/Engine/Plugins/Editor/PluginBrowser/Source/PluginBrowser/Private/SFilePathBlock.h
--- a/Engine/Plugins/Editor/PluginBrowser/Source/PluginBrowser/Private/SFilePathBlock.h
+++ b/Engine/Plugins/Editor/PluginBrowser/Source/PluginBrowser/Private/SFilePathBlock.h
@@ -22,9 +22,7 @@ class SFilePathBlock : public SCompoundWidget
 public:
 
 	SLATE_BEGIN_ARGS( SFilePathBlock )
-		: _LabelBackgroundColor(FLinearColor::Black)
-		, _LabelBackgroundBrush(FEditorStyle::GetBrush("WhiteBrush"))
-		, _ReadOnlyFolderPath(false)
+		: _ReadOnlyFolderPath(false)
 	{}
 		/** Attribute specifying the text to display in the folder input */
 		SLATE_ATTRIBUTE(FText, FolderPath)
@@ -35,12 +33,6 @@ public:
 		/** Hint name that appears when there is no text in the name box */
 		SLATE_ATTRIBUTE(FText, NameHint)
 
-		/** Background label tint for the folder/name labels */
-		SLATE_ATTRIBUTE(FSlateColor, LabelBackgroundColor)
-
-		/** Background label brush for the folder/name labels */
-		SLATE_ATTRIBUTE(const FSlateBrush*, LabelBackgroundBrush)
-
 		/** If true, the folder path cannot be modified by the user */
 		SLATE_ARGUMENT(bool, ReadOnlyFolderPath)
 
diff --git a/Engine/Plugins/Editor/PluginBrowser/Source/PluginBrowser/Private/SNewPluginWizard.cpp b/Engine/Plugins/Editor/PluginBrowser/Source/PluginBrowser/Private/SNewPluginWizard.cpp
--- a/Engine/Plugins/Editor/PluginBrowser/Source/PluginBrowser/Private/SNewPluginWizard.cpp
+++ b/Engine/Plugins/Editor/PluginBrowser/Source/PluginBrowser/Private/SNewPluginWizard.cpp
@@ -108,14 +108,14 @@ void SNewPluginWizard::Construct(const FArguments& Args, TSharedPtr<SDockTab> In
 	PluginFolderPath = AbsoluteGamePluginPath;
 	bIsPluginPathValid = true;
 
-	const float PaddingAmount = FPluginStyle::Get()->GetFloat("PluginCreator.Padding");
-
 	// Create the list view and ensure that it exists
 	GenerateListViewWidget();
 	check(ListView.IsValid());
 
 	TSharedPtr<SWidget> HeaderWidget = PluginWizardDefinition->GetCustomHeaderWidget();
 	FText PluginNameTextHint = PluginWizardDefinition->IsMod() ? LOCTEXT("ModNameTextHint", "Mod Name") : LOCTEXT("PluginNameTextHint", "Plugin Name");
+	
+	const FMargin PaddingAmount(5.0f);
 
 	TSharedRef<SVerticalBox> MainContent = SNew(SVerticalBox)
 	+SVerticalBox::Slot()
@@ -164,8 +164,6 @@ void SNewPluginWizard::Construct(const FArguments& Args, TSharedPtr<SDockTab> In
 	[
 		SAssignNew(FilePathBlock, SFilePathBlock)
 		.OnBrowseForFolder(this, &SNewPluginWizard::OnBrowseButtonClicked)
-		.LabelBackgroundBrush(FPluginStyle::Get()->GetBrush("PluginCreator.Background"))
-		.LabelBackgroundColor(FLinearColor::White)
 		.FolderPath(this, &SNewPluginWizard::GetPluginDestinationPath)
 		.Name(this, &SNewPluginWizard::GetCurrentPluginName)
 		.NameHint(PluginNameTextHint)
diff --git a/Engine/Plugins/Experimental/ModelingToolsEditorMode/Source/ModelingToolsEditorMode/Private/ModelingToolsEditorModeStyle.cpp b/Engine/Plugins/Experimental/ModelingToolsEditorMode/Source/ModelingToolsEditorMode/Private/ModelingToolsEditorModeStyle.cpp
--- a/Engine/Plugins/Experimental/ModelingToolsEditorMode/Source/ModelingToolsEditorMode/Private/ModelingToolsEditorModeStyle.cpp
+++ b/Engine/Plugins/Experimental/ModelingToolsEditorMode/Source/ModelingToolsEditorMode/Private/ModelingToolsEditorModeStyle.cpp
@@ -11,7 +11,6 @@
 
 
 #define IMAGE_PLUGIN_BRUSH( RelativePath, ... ) FSlateImageBrush( FModelingToolsEditorModeStyle::InContent( RelativePath, ".png" ), __VA_ARGS__ )
-#define DEFAULT_FONT(...) FCoreStyle::GetDefaultFontStyle(__VA_ARGS__)
 
 // This is to fix the issue that SlateStyleMacros like IMAGE_BRUSH look for RootToContentDir but StyleSet->RootToContentDir is how this style is set up
 #define RootToContentDir StyleSet->RootToContentDir
diff --git a/Engine/Source/Editor/EditorStyle/Private/StarshipStyle.cpp b/Engine/Source/Editor/EditorStyle/Private/StarshipStyle.cpp
--- a/Engine/Source/Editor/EditorStyle/Private/StarshipStyle.cpp
+++ b/Engine/Source/Editor/EditorStyle/Private/StarshipStyle.cpp
@@ -18,7 +18,6 @@
 #include "Framework/Application/SlateApplication.h"
 
 
-#define DEFAULT_FONT(...) FStarshipCoreStyle::GetDefaultFontStyle(__VA_ARGS__)
 #define ICON_FONT(...) FSlateFontInfo(RootToContentDir("Fonts/FontAwesome", TEXT(".ttf")), __VA_ARGS__)
 
 #define LOCTEXT_NAMESPACE "EditorStyle"
@@ -949,10 +948,6 @@ void FStarshipEditorStyle::FStyle::SetupGeneralStyles()
 		Set( "AnimEditor.EditPreviewParameters.Small", new IMAGE_BRUSH( "Icons/icon_adjust_parameters_40x", Icon20x20) );		
 	}
 
-
-	// Gamma reference.
-	Set("GammaReference",new IMAGE_BRUSH( "Old/GammaReference",FVector2D(256,128)));
-
 #endif // WITH_EDITOR || (IS_PROGRAM && WITH_UNREAL_DEVELOPER_TOOLS)
 
 	// Output Log Window
@@ -1009,22 +1004,11 @@ void FStarshipEditorStyle::FStyle::SetupGeneralStyles()
 		Set("OutputLog.OpenInExternalEditor", new IMAGE_BRUSH("Icons/icon_Asset_Open_In_External_Editor_16x", Icon16x16));
 
 	}
-
-	// Debug Tools Window
-	{
-		Set("DebugTools.TabIcon", new IMAGE_BRUSH( "Icons/icon_tab_DebugTools_16x", Icon16x16 ) );
-	}
-
 	// Performance Analysis Tools Window
 	{
 		Set("PerfTools.TabIcon", new IMAGE_BRUSH( "Icons/icon_tab_PerfTools_16x", Icon16x16 ) );
 	}
 
-	// Modules Window
-	{
-		Set("Modules.TabIcon", new IMAGE_BRUSH( "Icons/icon_tab_Modules_16px", Icon16x16 ) );
-	}
-
 	// Class Viewer Window
 	{
 		Set("ClassViewer.TabIcon", new IMAGE_BRUSH( "Icons/icon_tab_ClassViewer_16x", Icon16x16 ) );
@@ -2120,7 +2104,7 @@ void FStarshipEditorStyle::FStyle::SetupSequencerStyles()
 		Set("Sequencer.RestoreAnimatedState.Small", new IMAGE_BRUSH("Sequencer/Main_Icons/Icon_Sequencer_RestoreAnimatedState_24x", Icon24x24));
 		Set("Sequencer.GenericGripLeft", new BOX_BRUSH("Sequencer/GenericGripLeft", FMargin(5.f / 16.f)));
 		Set("Sequencer.GenericGripRight", new BOX_BRUSH("Sequencer/GenericGripRight", FMargin(5.f / 16.f)));
-		Set("Sequencer.SectionArea.Background", new FSlateColorBrush(FColor::White));
+		Set("Sequencer.SectionArea.Background", new FSlateColorBrush(FStyleColors::White));
 
 		Set("Sequencer.Section.Background", new BORDER_BRUSH(TEXT("Sequencer/SectionBackground"), FMargin(4.f / 16.f)));
 		Set("Sequencer.Section.BackgroundTint", new BOX_BRUSH(TEXT("Sequencer/SectionBackgroundTint"), FMargin(4 / 16.f)));
diff --git a/Engine/Source/Runtime/SlateCore/Public/Styling/SlateStyleMacros.h b/Engine/Source/Runtime/SlateCore/Public/Styling/SlateStyleMacros.h
--- a/Engine/Source/Runtime/SlateCore/Public/Styling/SlateStyleMacros.h
+++ b/Engine/Source/Runtime/SlateCore/Public/Styling/SlateStyleMacros.h
@@ -18,3 +18,4 @@
 #define CORE_BOX_BRUSH_SVG( RelativePath, ... ) FSlateVectorBoxBrush(RootToCoreContentDir(RelativePath, TEXT(".svg")), __VA_ARGS__)
 #define CORE_BORDER_BRUSH_SVG( RelativePath, ... ) FSlateVectorBorderBrush(RootToCoreContentDir(RelativePath, TEXT(".svg")), __VA_ARGS__)
 
+#define DEFAULT_FONT(...) FStarshipCoreStyle::GetDefaultFontStyle(__VA_ARGS__)
\ No newline at end of file
