Added an annotation in the plugin browser for NoRedist/NotForLicensees plugins (next to where the beta/experimental annotation goes)

diff --git a/Engine/Plugins/Editor/PluginBrowser/Source/PluginBrowser/Private/SPluginTile.cpp b/Engine/Plugins/Editor/PluginBrowser/Source/PluginBrowser/Private/SPluginTile.cpp
--- a/Engine/Plugins/Editor/PluginBrowser/Source/PluginBrowser/Private/SPluginTile.cpp
+++ b/Engine/Plugins/Editor/PluginBrowser/Source/PluginBrowser/Private/SPluginTile.cpp
@@ -198,6 +198,22 @@ void SPluginTile::RecreateWidgets()
 		}
 	}
 
+	TSharedRef<SWidget> RestrictedPluginWidget = SNullWidget::NullWidget;
+	if (FPaths::IsRestrictedPath(Plugin->GetDescriptorFileName()))
+	{
+		RestrictedPluginWidget = SNew(SHorizontalBox)
+			+ SHorizontalBox::Slot()
+				.AutoWidth()
+				.VAlign(VAlign_Bottom)
+				.Padding(2.0f, 0.0f, 8.0f, 1.0f)
+				[
+					SNew(STextBlock)
+						.TextStyle(FPluginStyle::Get(), "PluginTile.BetaText")
+						.Text(LOCTEXT("PluginRestrictedText", "[Restricted]"))
+						.ToolTipText(FText::AsCultureInvariant(Plugin->GetDescriptorFileName()))
+				];
+	}
+
 	ChildSlot
 	[
 		SNew(SBorder)
@@ -279,6 +295,14 @@ void SPluginTile::RecreateWidgets()
 											[
 												SNew(SHorizontalBox)
 
+												// noredist/restricted label
+												+ SHorizontalBox::Slot()
+													.AutoWidth()
+													.VAlign(VAlign_Bottom)
+													[
+														RestrictedPluginWidget
+													]
+
 												// beta version label
 												+ SHorizontalBox::Slot()
 													.AutoWidth()
