Added a bit more width to the plugin browser SFilePathBlock so it handles typical depot project paths better

diff --git a/Engine/Plugins/Editor/PluginBrowser/Source/PluginBrowser/Private/SFilePathBlock.cpp b/Engine/Plugins/Editor/PluginBrowser/Source/PluginBrowser/Private/SFilePathBlock.cpp
--- a/Engine/Plugins/Editor/PluginBrowser/Source/PluginBrowser/Private/SFilePathBlock.cpp
+++ b/Engine/Plugins/Editor/PluginBrowser/Source/PluginBrowser/Private/SFilePathBlock.cpp
@@ -28,14 +28,21 @@ void SFilePathBlock::Construct(const FArguments& InArgs)
 		BrowseForFolderToolTipText = LOCTEXT("BrowseForFolder", "Browse for a folder");
 	}
 
+	const float DesiredWidth = 700.0f;
+	const float DesiredRatioForPath = 2.0f;
+
+	const float DerivedMinWidthForPath = DesiredWidth * DesiredRatioForPath / (DesiredRatioForPath + 1.0f);
+	const float DerivedMinWidthForName = DesiredWidth - DerivedMinWidthForPath;
+
 	ChildSlot
 	[
 		SNew(SHorizontalBox)
 		// Folder input
 		+SHorizontalBox::Slot()
+		.FillWidth(DesiredRatioForPath)
 		[
 			SNew(SBox)
-			.MinDesiredWidth(250.f)
+			.MinDesiredWidth(DerivedMinWidthForPath)
 			[
 				SNew(SOverlay)
 				+ SOverlay::Slot()
@@ -71,7 +78,7 @@ void SFilePathBlock::Construct(const FArguments& InArgs)
 		.VAlign(VAlign_Center)
 		[
 			SNew(SBox)
-			.MinDesiredWidth(200.f)
+			.MinDesiredWidth(DerivedMinWidthForName)
 			[
 				SAssignNew(NameTextBox, SEditableTextBox)
 				.Text(InArgs._Name)
