Sorted project plugins above installed plugins

diff --git a/Engine/Plugins/Editor/PluginBrowser/Source/PluginBrowser/Private/SPluginCategoryTree.cpp b/Engine/Plugins/Editor/PluginBrowser/Source/PluginBrowser/Private/SPluginCategoryTree.cpp
--- a/Engine/Plugins/Editor/PluginBrowser/Source/PluginBrowser/Private/SPluginCategoryTree.cpp
+++ b/Engine/Plugins/Editor/PluginBrowser/Source/PluginBrowser/Private/SPluginCategoryTree.cpp
@@ -209,13 +209,13 @@ void SPluginCategoryTree::RebuildAndFilterCategoryTree()
 	{
 		RootCategories.Add(ModCategory);
 	}
-	if(InstalledCategory->SubCategories.Num() > 0 || InstalledCategory->Plugins.Num() > 0)
+	if(ProjectCategory->SubCategories.Num() > 0 || ProjectCategory->Plugins.Num() > 0)
 	{
-		RootCategories.Add(InstalledCategory);
+		RootCategories.Add(ProjectCategory);
 	}
-	if (ProjectCategory->SubCategories.Num() > 0 || ProjectCategory->Plugins.Num() > 0)
+	if(InstalledCategory->SubCategories.Num() > 0 || InstalledCategory->Plugins.Num() > 0)
 	{
-		RootCategories.Add(ProjectCategory);
+		RootCategories.Add(InstalledCategory);
 	}
 	if(BuiltInCategory->SubCategories.Num() > 0 || BuiltInCategory->Plugins.Num() > 0)
 	{
