Making the plugin wizard extensible from other plugins
- Added RegisterPluginTemplate / UnegisterPluginTemplate to the modular feature IPluginsEditorFeature
- Added sorting, path validation, and post-creation extension points to FPluginTemplateDescription
- Devolved some choices from IPluginWizardDefinition to FPluginTemplateDescription
- Fixed a sticky width bug on the plugin wizard instructions prompt

diff --git a/Engine/Plugins/Editor/PluginBrowser/Source/PluginBrowser/Private/DefaultPluginWizardDefinition.cpp b/Engine/Plugins/Editor/PluginBrowser/Source/PluginBrowser/Private/DefaultPluginWizardDefinition.cpp
--- a/Engine/Plugins/Editor/PluginBrowser/Source/PluginBrowser/Private/DefaultPluginWizardDefinition.cpp
+++ b/Engine/Plugins/Editor/PluginBrowser/Source/PluginBrowser/Private/DefaultPluginWizardDefinition.cpp
@@ -6,6 +6,7 @@
 #include "Interfaces/IPluginManager.h"
 #include "HAL/PlatformFileManager.h"
 #include "Misc/App.h"
+#include "PluginBrowserModule.h"
 
 #define LOCTEXT_NAMESPACE "NewPluginWizard"
 
@@ -35,17 +36,46 @@ void FDefaultPluginWizardDefinition::PopulateTemplatesSource()
 	const FText EditorModeDescription = LOCTEXT("EditorModeDesc", "Create a plugin that will have an editor mode.\n\nThis will include a toolkit example to specify UI that will appear in \"Modes\" tab (next to Foliage, Landscape etc).\nIt will also include very basic UI that demonstrates editor interaction and undo/redo functions usage.");
 	const FText ThirdPartyDescription = LOCTEXT("ThirdPartyDesc", "Create a plugin that uses an included third party library.\n\nThis can be used as an example of how to include, load and use a third party library yourself.");
 
-	TemplateDefinitions.Add(MakeShareable(new FPluginTemplateDescription(ContentOnlyTemplateName, ContentOnlyDescription, TEXT("ContentOnly"), true, EHostType::Runtime)));
+	TSharedRef<FPluginTemplateDescription> ContentOnlyTemplate = MakeShareable(new FPluginTemplateDescription(ContentOnlyTemplateName, ContentOnlyDescription, PluginBaseDir / TEXT("Templates") / TEXT("ContentOnly"), true, EHostType::Runtime));
+	ContentOnlyTemplate->SortPriority = 1;
+	TemplateDefinitions.Add(ContentOnlyTemplate);
+
 	if (!bIsContentOnlyProject)
 	{
 		// Insert the blank template to make sure it appears before the content only template.
-		TemplateDefinitions.Insert(MakeShareable(new FPluginTemplateDescription(BlankTemplateName, BlankDescription, TEXT("Blank"), true, EHostType::Runtime)), 0);
-		TemplateDefinitions.Add(MakeShareable(new FPluginTemplateDescription(BlueprintLibTemplateName, BlueprintLibDescription, TEXT("BlueprintLibrary"), true, EHostType::Runtime, ELoadingPhase::PreLoadingScreen)));
-		TemplateDefinitions.Add(MakeShareable(new FPluginTemplateDescription(BasicTemplateName, BasicDescription, TEXT("Basic"), false, EHostType::Editor)));
-		TemplateDefinitions.Add(MakeShareable(new FPluginTemplateDescription(AdvancedTemplateName, AdvancedDescription, TEXT("Advanced"), false, EHostType::Editor)));
-		TemplateDefinitions.Add(MakeShareable(new FPluginTemplateDescription(EditorModeTemplateName, EditorModeDescription, TEXT("EditorMode"), false, EHostType::Editor)));
-		TemplateDefinitions.Add(MakeShareable(new FPluginTemplateDescription(ThirdPartyTemplateName, ThirdPartyDescription, TEXT("ThirdPartyLibrary"), true, EHostType::Runtime)));
+
+		TSharedRef<FPluginTemplateDescription> BlankTemplate = MakeShareable(new FPluginTemplateDescription(BlankTemplateName, BlankDescription, PluginBaseDir / TEXT("Templates") / TEXT("Blank"), true, EHostType::Runtime));
+		BlankTemplate->SortPriority = 2;
+		TemplateDefinitions.Add(BlankTemplate);
+
+		TemplateDefinitions.Add(MakeShareable(new FPluginTemplateDescription(BlueprintLibTemplateName, BlueprintLibDescription, PluginBaseDir / TEXT("Templates") / TEXT("BlueprintLibrary"), true, EHostType::Runtime, ELoadingPhase::PreLoadingScreen)));
+		TemplateDefinitions.Add(MakeShareable(new FPluginTemplateDescription(BasicTemplateName, BasicDescription, PluginBaseDir / TEXT("Templates") / TEXT("Basic"), false, EHostType::Editor)));
+		TemplateDefinitions.Add(MakeShareable(new FPluginTemplateDescription(AdvancedTemplateName, AdvancedDescription, PluginBaseDir / TEXT("Templates") / TEXT("Advanced"), false, EHostType::Editor)));
+		TemplateDefinitions.Add(MakeShareable(new FPluginTemplateDescription(EditorModeTemplateName, EditorModeDescription, PluginBaseDir / TEXT("Templates") / TEXT("EditorMode"), false, EHostType::Editor)));
+		TemplateDefinitions.Add(MakeShareable(new FPluginTemplateDescription(ThirdPartyTemplateName, ThirdPartyDescription, PluginBaseDir / TEXT("Templates") / TEXT("ThirdPartyLibrary"), true, EHostType::Runtime)));
 	}
+
+	// Add external templates that came from the modular feature interface (e.g., from another plugin like Game Features)
+	TemplateDefinitions.Append(FPluginBrowserModule::Get().GetAddedPluginTemplates());
+
+	// Don't show the option to make an engine plugin in installed builds
+	const bool bAllowEnginePlugins = !FApp::IsEngineInstalled();
+	for (const TSharedRef<FPluginTemplateDescription>& Template : TemplateDefinitions)
+	{
+		Template->bCanBePlacedInEngine = Template->bCanBePlacedInEngine && bAllowEnginePlugins;
+	}
+
+	TemplateDefinitions.Sort([](const TSharedRef<FPluginTemplateDescription>& A, const TSharedRef<FPluginTemplateDescription>& B)
+	{
+		if (A->SortPriority != B->SortPriority)
+		{
+			return A->SortPriority > B->SortPriority;
+		}
+		else
+		{
+			return A->Name.CompareTo(B->Name) <= 0;
+		}
+	});
 }
 
 const TArray<TSharedRef<FPluginTemplateDescription>>& FDefaultPluginWizardDefinition::GetTemplatesSource() const
@@ -54,25 +84,14 @@ const TArray<TSharedRef<FPluginTemplateDescription>>& FDefaultPluginWizardDefini
 }
 
 
-void FDefaultPluginWizardDefinition::OnTemplateSelectionChanged(TArray<TSharedRef<FPluginTemplateDescription>> InSelectedItems, ESelectInfo::Type SelectInfo)
+void FDefaultPluginWizardDefinition::OnTemplateSelectionChanged(TSharedPtr<FPluginTemplateDescription> InSelectedItem, ESelectInfo::Type SelectInfo)
 {
-	CurrentTemplateDefinition.Reset();
-	
-	if (InSelectedItems.Num() > 0)
-	{
-		CurrentTemplateDefinition = InSelectedItems[0];
-	}
+	CurrentTemplateDefinition = InSelectedItem;
 }
 
-TArray<TSharedPtr<FPluginTemplateDescription>> FDefaultPluginWizardDefinition::GetSelectedTemplates() const
+TSharedPtr<FPluginTemplateDescription> FDefaultPluginWizardDefinition::GetSelectedTemplate() const
 {
-	TArray<TSharedPtr<FPluginTemplateDescription>> Selection;
-	if (CurrentTemplateDefinition.IsValid())
-	{
-		Selection.Add(CurrentTemplateDefinition);
-	}
-
-	return Selection;
+	return CurrentTemplateDefinition;
 }
 
 bool FDefaultPluginWizardDefinition::HasValidTemplateSelection() const
@@ -85,17 +104,6 @@ void FDefaultPluginWizardDefinition::ClearTemplateSelection()
 	CurrentTemplateDefinition.Reset();
 }
 
-bool FDefaultPluginWizardDefinition::AllowsEnginePlugins() const
-{
-	// Don't show the option to make an engine plugin in installed builds
-	return !FApp::IsEngineInstalled();
-}
-
-bool FDefaultPluginWizardDefinition::CanContainContent() const
-{
-	return CurrentTemplateDefinition.IsValid() ? CurrentTemplateDefinition->bCanContainContent : false;
-}
-
 bool FDefaultPluginWizardDefinition::HasModules() const
 {
 	FString SourceFolderPath = GetPluginFolderPath() / TEXT("Source");
@@ -181,9 +189,7 @@ FString FDefaultPluginWizardDefinition::GetPluginFolderPath() const
 
 FString FDefaultPluginWizardDefinition::GetFolderForTemplate(TSharedRef<FPluginTemplateDescription> InTemplate) const
 {
-	FString TemplateFolderName = PluginBaseDir / TEXT("Templates") / InTemplate->OnDiskPath;
-
-	return TemplateFolderName;
+	return InTemplate->OnDiskPath;
 }
 
 #undef LOCTEXT_NAMESPACE
\ No newline at end of file
diff --git a/Engine/Plugins/Editor/PluginBrowser/Source/PluginBrowser/Private/DefaultPluginWizardDefinition.h b/Engine/Plugins/Editor/PluginBrowser/Source/PluginBrowser/Private/DefaultPluginWizardDefinition.h
--- a/Engine/Plugins/Editor/PluginBrowser/Source/PluginBrowser/Private/DefaultPluginWizardDefinition.h
+++ b/Engine/Plugins/Editor/PluginBrowser/Source/PluginBrowser/Private/DefaultPluginWizardDefinition.h
@@ -11,15 +11,12 @@ public:
 
 	// Begin IPluginWizardDefinition interface
 	virtual const TArray<TSharedRef<FPluginTemplateDescription>>& GetTemplatesSource() const override;
-	virtual void OnTemplateSelectionChanged(TArray<TSharedRef<FPluginTemplateDescription>> InSelectedItems, ESelectInfo::Type SelectInfo) override;
-	virtual TArray<TSharedPtr<FPluginTemplateDescription>> GetSelectedTemplates() const override;
+	virtual void OnTemplateSelectionChanged(TSharedPtr<FPluginTemplateDescription> InSelectedItem, ESelectInfo::Type SelectInfo) override;
+	virtual TSharedPtr<FPluginTemplateDescription> GetSelectedTemplate() const override;
 	virtual void ClearTemplateSelection() override;
 	virtual bool HasValidTemplateSelection() const override;
 
-	virtual ESelectionMode::Type GetSelectionMode() const override { return ESelectionMode::Single; }
-	virtual bool AllowsEnginePlugins() const override;
 	virtual bool CanShowOnStartup() const override { return false; }
-	virtual bool CanContainContent() const override;
 	virtual bool HasModules() const override;
 	virtual bool IsMod() const override;
 	virtual void OnShowOnStartupCheckboxChanged(ECheckBoxState CheckBoxState) override {}
diff --git a/Engine/Plugins/Editor/PluginBrowser/Source/PluginBrowser/Private/PluginBrowserModule.cpp b/Engine/Plugins/Editor/PluginBrowser/Source/PluginBrowser/Private/PluginBrowserModule.cpp
--- a/Engine/Plugins/Editor/PluginBrowser/Source/PluginBrowser/Private/PluginBrowserModule.cpp
+++ b/Engine/Plugins/Editor/PluginBrowser/Source/PluginBrowser/Private/PluginBrowserModule.cpp
@@ -98,6 +98,16 @@ void FPluginBrowserModule::ShutdownModule()
 	IModularFeatures::Get().UnregisterModularFeature( EditorFeatures::PluginsEditor, this );
 }
 
+void FPluginBrowserModule::RegisterPluginTemplate(TSharedRef<FPluginTemplateDescription> Template)
+{
+	AddedPluginTemplates.Add(Template);
+}
+
+void FPluginBrowserModule::UnregisterPluginTemplate(TSharedRef<FPluginTemplateDescription> Template)
+{
+	AddedPluginTemplates.RemoveSingle(Template);
+}
+
 void FPluginBrowserModule::SetPluginPendingEnableState(const FString& PluginName, bool bCurrentlyEnabled, bool bPendingEnabled)
 {
 	if (bCurrentlyEnabled == bPendingEnabled)
diff --git a/Engine/Plugins/Editor/PluginBrowser/Source/PluginBrowser/Private/PluginBrowserModule.h b/Engine/Plugins/Editor/PluginBrowser/Source/PluginBrowser/Private/PluginBrowserModule.h
--- a/Engine/Plugins/Editor/PluginBrowser/Source/PluginBrowser/Private/PluginBrowserModule.h
+++ b/Engine/Plugins/Editor/PluginBrowser/Source/PluginBrowser/Private/PluginBrowserModule.h
@@ -25,6 +25,10 @@ public:
 	virtual void StartupModule() override;
 	virtual void ShutdownModule() override;
 
+	/** IPluginsEditorFeature implementation */
+	virtual void RegisterPluginTemplate(TSharedRef<FPluginTemplateDescription> Template) override;
+	virtual void UnregisterPluginTemplate(TSharedRef<FPluginTemplateDescription> Template) override;
+
 	/** Gets a delegate so that you can register/unregister to receive callbacks when plugins are created */
 	FOnNewPluginCreated& OnNewPluginCreated() {return NewPluginCreatedDelegate;}
 
@@ -67,6 +71,8 @@ public:
 	/** Spawns the plugin creator tab with a specific wizard definition */
 	virtual TSharedRef<SDockTab> SpawnPluginCreatorTab(const FSpawnTabArgs& SpawnTabArgs, TSharedPtr<IPluginWizardDefinition> PluginWizardDefinition) override;
 
+	const TArray<TSharedRef<FPluginTemplateDescription>>& GetAddedPluginTemplates() const { return AddedPluginTemplates; }
+
 private:
 	/** Called to spawn the plugin browser tab */
 	TSharedRef<SDockTab> HandleSpawnPluginBrowserTab(const FSpawnTabArgs& SpawnTabArgs);
@@ -86,6 +92,9 @@ private:
 	/** Updates the user's config file with the list of installed plugins that they've seen. */
 	void UpdatePreviousInstalledPlugins();
 
+	/** List of added plugin templates */
+	TArray<TSharedRef<FPluginTemplateDescription>> AddedPluginTemplates;
+
 	/** The spawned browser tab */
 	TWeakPtr<SDockTab> PluginBrowserTab;
 
diff --git a/Engine/Plugins/Editor/PluginBrowser/Source/PluginBrowser/Private/SNewPluginWizard.cpp b/Engine/Plugins/Editor/PluginBrowser/Source/PluginBrowser/Private/SNewPluginWizard.cpp
--- a/Engine/Plugins/Editor/PluginBrowser/Source/PluginBrowser/Private/SNewPluginWizard.cpp
+++ b/Engine/Plugins/Editor/PluginBrowser/Source/PluginBrowser/Private/SNewPluginWizard.cpp
@@ -58,7 +58,7 @@ static bool IsContentOnlyProject()
 SNewPluginWizard::SNewPluginWizard()
 	: bIsPluginPathValid(false)
 	, bIsPluginNameValid(false)
-	, bIsEnginePlugin(false)
+	, bIsSelectedPathInEngine(false)
 {
 	AbsoluteGamePluginPath = IFileManager::Get().ConvertToAbsolutePathForExternalAppForWrite(*FPaths::ProjectPluginsDir());
 	FPaths::MakePlatformFilename(AbsoluteGamePluginPath);
@@ -136,7 +136,7 @@ void SNewPluginWizard::Construct(const FArguments& Args, TSharedPtr<SDockTab> In
 		+ SHorizontalBox::Slot()
 		.FillWidth(1.0f)
 		.Padding(PaddingAmount)
-		.HAlign(HAlign_Left)
+		.HAlign(HAlign_Fill)
 		[
 			SNew(SVerticalBox)
 
@@ -169,7 +169,6 @@ void SNewPluginWizard::Construct(const FArguments& Args, TSharedPtr<SDockTab> In
 		.NameHint(PluginNameTextHint)
 		.OnFolderChanged(this, &SNewPluginWizard::OnFolderPathTextChanged)
 		.OnNameChanged(this, &SNewPluginWizard::OnPluginNameTextChanged)
-        .ReadOnlyFolderPath( !PluginWizardDefinition->AllowsEnginePlugins() )	// only allow the user to select the folder if they can create engine plugins
 	];
 
 	// Add the descriptor data object if it exists
@@ -185,28 +184,30 @@ void SNewPluginWizard::Construct(const FArguments& Args, TSharedPtr<SDockTab> In
 		];
 	}
 
-	if (PluginWizardDefinition->AllowsEnginePlugins())
-	{
-		MainContent->AddSlot()
-		.AutoHeight()
+	MainContent->AddSlot()
+	.AutoHeight()
+	[
+		SNew(SBox)
+		.HAlign(HAlign_Left)
+		.VAlign(VAlign_Center)
 		.Padding(PaddingAmount)
 		[
-			SNew(SBox)
-			.HAlign(HAlign_Left)
-			.VAlign(VAlign_Center)
+			SNew(SCheckBox)
+			.OnCheckStateChanged(this, &SNewPluginWizard::OnEnginePluginCheckboxChanged)
+			.IsChecked(this, &SNewPluginWizard::IsEnginePlugin)
+			.Visibility_Lambda([=]()
+			{
+				FPluginTemplateDescription* Template = PluginWizardDefinition->GetSelectedTemplate().Get();
+				return ((Template != nullptr) && Template->bCanBePlacedInEngine) ? EVisibility::Visible : EVisibility::Collapsed;
+			})
+			.ToolTipText(LOCTEXT("EnginePluginButtonToolTip", "Toggles whether this plugin will be created in the current project or the engine directory."))
+			.Content()
 			[
-				SNew(SCheckBox)
-				.OnCheckStateChanged(this, &SNewPluginWizard::OnEnginePluginCheckboxChanged)
-				.IsChecked(this, &SNewPluginWizard::IsEnginePlugin)
-				.ToolTipText(LOCTEXT("EnginePluginButtonToolTip", "Toggles whether this plugin will be created in the current project or the engine directory."))
-				.Content()
-				[
-					SNew(STextBlock)
-					.Text(LOCTEXT("EnginePluginCheckbox", "Is Engine Plugin"))
-				]
+				SNew(STextBlock)
+				.Text(LOCTEXT("EnginePluginCheckbox", "Is Engine Plugin"))
 			]
-		];
-	}
+		]
+	];
 
 	if (PluginWizardDefinition->CanShowOnStartup())
 	{
@@ -242,7 +243,11 @@ void SNewPluginWizard::Construct(const FArguments& Args, TSharedPtr<SDockTab> In
 		[
 			SAssignNew(ShowPluginContentDirectoryCheckBox, SCheckBox)
 			.IsChecked(ECheckBoxState::Checked)
-			.Visibility(this, &SNewPluginWizard::GetShowPluginContentDirectoryVisibility)
+			.Visibility_Lambda([=]()
+			{
+				FPluginTemplateDescription* Template = PluginWizardDefinition->GetSelectedTemplate().Get();
+				return ((Template != nullptr) && Template->bCanContainContent) ? EVisibility::Visible : EVisibility::Collapsed;
+			})
 			.ToolTipText(LOCTEXT("ShowPluginContentDirectoryToolTip", "Shows the content directory after creation."))
 			.Content()
 			[
@@ -277,44 +282,14 @@ void SNewPluginWizard::Construct(const FArguments& Args, TSharedPtr<SDockTab> In
 
 void SNewPluginWizard::GenerateListViewWidget()
 {
-	// for now, just determine what view to create based on the selection mode of the wizard definition
-	ESelectionMode::Type SelectionMode = PluginWizardDefinition->GetSelectionMode();
-
 	// Get the source of the templates to use for the list view
 	const TArray<TSharedRef<FPluginTemplateDescription>>& TemplateSource = PluginWizardDefinition->GetTemplatesSource();
 
-	switch (SelectionMode)
-	{
-		case ESelectionMode::Multi:
-		{
-			ListView = SNew(STileView<TSharedRef<FPluginTemplateDescription>>)
-				.SelectionMode(SelectionMode)
-				.ListItemsSource(&TemplateSource)
-				.OnGenerateTile(this, &SNewPluginWizard::OnGenerateTemplateTile)
-				.OnSelectionChanged(this, &SNewPluginWizard::OnTemplateSelectionChanged)
-				.ItemHeight(180.0f);
-		}
-		break;
-
-		case ESelectionMode::Single:
-		case ESelectionMode::SingleToggle:
-		{
-			ListView = SNew(SListView<TSharedRef<FPluginTemplateDescription>>)
-				.SelectionMode(SelectionMode)
-				.ListItemsSource(&TemplateSource)
-				.OnGenerateRow(this, &SNewPluginWizard::OnGenerateTemplateRow)
-				.OnSelectionChanged(this, &SNewPluginWizard::OnTemplateSelectionChanged);
-		}
-		break;
-
-		case ESelectionMode::None:
-		default:
-		{
-			// This isn't a valid selection mode for this widget
-			check(false);
-		}
-		break;
-	}		
+	ListView = SNew(SListView<TSharedRef<FPluginTemplateDescription>>)
+		.SelectionMode(ESelectionMode::Single)
+		.ListItemsSource(&TemplateSource)
+		.OnGenerateRow(this, &SNewPluginWizard::OnGenerateTemplateRow)
+		.OnSelectionChanged(this, &SNewPluginWizard::OnTemplateSelectionChanged);
 }
 
 void SNewPluginWizard::GeneratePluginTemplateDynamicBrush(TSharedRef<FPluginTemplateDescription> InItem)
@@ -459,17 +434,37 @@ TSharedRef<ITableRow> SNewPluginWizard::OnGenerateTemplateRow(TSharedRef<FPlugin
 void SNewPluginWizard::OnTemplateSelectionChanged(TSharedPtr<FPluginTemplateDescription> InItem, ESelectInfo::Type SelectInfo)
 {
 	// Forward the set of selected items to the plugin wizard definition
-	TArray<TSharedRef<FPluginTemplateDescription>> SelectedItems;
+	TSharedPtr<FPluginTemplateDescription> SelectedItem;
 
 	if (ListView.IsValid())
 	{
-		SelectedItems = ListView->GetSelectedItems();
+		if (ListView->GetSelectedItems().Num() > 0)
+		{
+			SelectedItem = ListView->GetSelectedItems()[0];
+		}
 	}
 
 	if (PluginWizardDefinition.IsValid())
 	{
-		PluginWizardDefinition->OnTemplateSelectionChanged(SelectedItems, SelectInfo);
+		TSharedPtr<FPluginTemplateDescription> OldSelectedItem = PluginWizardDefinition->GetSelectedTemplate();
+
+		if (OldSelectedItem != SelectedItem)
+		{
+			if (OldSelectedItem.IsValid())
+			{
+				OldSelectedItem->UpdatePathWhenTemplateUnselected(PluginFolderPath);
+			}
+
+			PluginWizardDefinition->OnTemplateSelectionChanged(SelectedItem, SelectInfo);
+
+			if (SelectedItem.IsValid())
+			{
+				SelectedItem->UpdatePathWhenTemplateSelected(PluginFolderPath);
+			}
+		}
 	}
+
+	ValidateFullPluginPath();
 }
 
 void SNewPluginWizard::OnFolderPathTextChanged(const FText& InText)
@@ -532,26 +527,50 @@ void SNewPluginWizard::ValidateFullPluginPath()
 
 	if (bIsNewPathValid)
 	{
-		bool bFoundValidPath = false;
+		bool bFoundKnownPath = false;
+		bool bPathIsInEngine = false;
 		FString AbsolutePath = IFileManager::Get().ConvertToAbsolutePathForExternalAppForWrite(*GetPluginDestinationPath().ToString());
 		FPaths::MakePlatformFilename(AbsolutePath);
 
 		if (AbsolutePath.StartsWith(AbsoluteGamePluginPath))
 		{
-			bFoundValidPath = true;
-			bIsEnginePlugin = false;
+			bFoundKnownPath = true;
+			bPathIsInEngine = false;
 		}
-		else if (!bFoundValidPath && !FApp::IsEngineInstalled())
+		else if (AbsolutePath.StartsWith(AbsoluteEnginePluginPath))
 		{
-			if (AbsolutePath.StartsWith(AbsoluteEnginePluginPath))
+			bFoundKnownPath = true;
+			bPathIsInEngine = true;
+		}
+		else
+		{
+			// Path is neither engine nor game, so this path will be added to the additional plugin directories for
+			// the project when the plugin is created
+		}
+
+		bIsSelectedPathInEngine = bPathIsInEngine;
+
+		// Run this path by the current template to make sure it's OK with it
+		TSharedPtr<FPluginTemplateDescription> Template = PluginWizardDefinition->GetSelectedTemplate();
+		if (Template.IsValid() && bIsNewPathValid)
+		{
+			if (bFoundKnownPath && bPathIsInEngine && !Template->bCanBePlacedInEngine)
 			{
-				bFoundValidPath = true;
-				bIsEnginePlugin = true;
+				bIsNewPathValid = false;
+				if (FApp::IsInstalled())
+				{
+					FolderPathError = LOCTEXT("TemplateCannotBeInEngine_InstalledBuild", "Plugins created in installed builds cannot be placed in the engine folder");
+				}
+				else
+				{
+					FolderPathError = LOCTEXT("TemplateCannotBeInEngine", "Plugins created from this template cannot be placed in the engine folder");
+				}
 			}
 		}
-		else
+
+		if (bIsNewPathValid && Template.IsValid())
 		{
-			// This path will be added to the additional plugin directories for the project when created
+			bIsNewPathValid = Template->ValidatePathForPlugin(AbsolutePath, /*out*/ FolderPathError);
 		}
 	}
 
@@ -581,16 +600,16 @@ FText SNewPluginWizard::GetCurrentPluginName() const
 
 ECheckBoxState SNewPluginWizard::IsEnginePlugin() const
 {
-	return bIsEnginePlugin ? ECheckBoxState::Checked : ECheckBoxState::Unchecked;
+	return bIsSelectedPathInEngine ? ECheckBoxState::Checked : ECheckBoxState::Unchecked;
 }
 
 void SNewPluginWizard::OnEnginePluginCheckboxChanged(ECheckBoxState NewCheckedState)
 {
 	bool bNewEnginePluginState = NewCheckedState == ECheckBoxState::Checked;
-	if (bIsEnginePlugin != bNewEnginePluginState)
+	if (bIsSelectedPathInEngine != bNewEnginePluginState)
 	{
-		bIsEnginePlugin = bNewEnginePluginState;
-		if (bIsEnginePlugin)
+		bIsSelectedPathInEngine = bNewEnginePluginState;
+		if (bIsSelectedPathInEngine)
 		{
 			PluginFolderPath = AbsoluteEnginePluginPath;
 		}
@@ -598,8 +617,8 @@ void SNewPluginWizard::OnEnginePluginCheckboxChanged(ECheckBoxState NewCheckedSt
 		{
 			PluginFolderPath = AbsoluteGamePluginPath;
 		}
-		bIsPluginPathValid = true;
-		FilePathBlock->SetFolderPathError(FText::GetEmpty());
+
+		ValidateFullPluginPath();
 	}
 }
 
@@ -611,12 +630,19 @@ FReply SNewPluginWizard::OnCreatePluginClicked()
 		return FReply::Unhandled();
 	}
 
+	TSharedPtr<FPluginTemplateDescription> Template = PluginWizardDefinition->GetSelectedTemplate();
+	if (!ensure(Template.IsValid()))
+	{
+		return FReply::Unhandled();
+	}
+
+
 	const FString PluginName = PluginNameText.ToString();
 	const bool bHasModules = PluginWizardDefinition->HasModules();
-
+	
 	FPluginUtils::FNewPluginParams CreationParams;
 	CreationParams.TemplateFolders = PluginWizardDefinition->GetFoldersForSelection();
-	CreationParams.bCanContainContent = PluginWizardDefinition->CanContainContent();
+	CreationParams.bCanContainContent = Template->bCanContainContent;
 	CreationParams.bHasModules = bHasModules;
 	CreationParams.ModuleDescriptorType = PluginWizardDefinition->GetPluginModuleDescriptor();
 	CreationParams.LoadingPhase = PluginWizardDefinition->GetPluginLoadingPhase();
@@ -638,10 +664,16 @@ FReply SNewPluginWizard::OnCreatePluginClicked()
 	TSharedPtr<IPlugin> NewPlugin = FPluginUtils::CreateAndMountNewPlugin(PluginName, PluginFolderPath, CreationParams, MountParams, FailReason);
 	const bool bSucceeded = NewPlugin.IsValid();
 
+
 	PluginWizardDefinition->PluginCreated(PluginName, bSucceeded);
 
 	if (bSucceeded)
 	{
+		//@TODO: Add the plugin files to source control if the project is configured for it!
+
+		// Let the template create additional assets / modify state after creation
+		Template->OnPluginCreated(NewPlugin);
+
 		// Notify that a new plugin has been created
 		FPluginBrowserModule& PluginBrowserModule = FPluginBrowserModule::Get();
 		PluginBrowserModule.BroadcastNewPluginCreated();
@@ -668,16 +700,4 @@ FReply SNewPluginWizard::OnCreatePluginClicked()
 	}
 }
 
-EVisibility SNewPluginWizard::GetShowPluginContentDirectoryVisibility() const
-{
-	EVisibility ShowContentVisibility = EVisibility::Collapsed;
-
-	if (PluginWizardDefinition->CanContainContent())
-	{
-		ShowContentVisibility = EVisibility::Visible;
-	}
-
-	return ShowContentVisibility;
-}
-
 #undef LOCTEXT_NAMESPACE
diff --git a/Engine/Plugins/Editor/PluginBrowser/Source/PluginBrowser/Private/SNewPluginWizard.h b/Engine/Plugins/Editor/PluginBrowser/Source/PluginBrowser/Private/SNewPluginWizard.h
--- a/Engine/Plugins/Editor/PluginBrowser/Source/PluginBrowser/Private/SNewPluginWizard.h
+++ b/Engine/Plugins/Editor/PluginBrowser/Source/PluginBrowser/Private/SNewPluginWizard.h
@@ -100,11 +100,6 @@ private:
 	 */
 	ECheckBoxState IsEnginePlugin() const;
 
-	/**
-	 * Returns the visibility of the "Open Content Directory" checkbox, which should be displayed for any plugin that can contain content
-	 */
-	EVisibility GetShowPluginContentDirectoryVisibility() const;
-
 	/**
 	 * Called when state of Engine plugin checkbox changes
 	 * @param NewCheckedState New state of the checkbox
@@ -165,7 +160,7 @@ private:
 	bool bIsPluginNameValid;
 
 	/** Whether we want to create a plugin in the engine folder */
-	bool bIsEnginePlugin;
+	bool bIsSelectedPathInEngine;
 
 	/** Tab that owns this wizard so that we can ask to close after completion */
 	TWeakPtr<SDockTab> OwnerTab;
diff --git a/Engine/Plugins/Editor/PluginBrowser/Source/PluginBrowser/Public/IPluginWizardDefinition.h b/Engine/Plugins/Editor/PluginBrowser/Source/PluginBrowser/Public/IPluginWizardDefinition.h
--- a/Engine/Plugins/Editor/PluginBrowser/Source/PluginBrowser/Public/IPluginWizardDefinition.h
+++ b/Engine/Plugins/Editor/PluginBrowser/Source/PluginBrowser/Public/IPluginWizardDefinition.h
@@ -7,48 +7,7 @@
 #include "ModuleDescriptor.h"
 #include "Styling/SlateTypes.h"
 #include "Types/SlateEnums.h"
-
-
-struct FSlateDynamicImageBrush;
-
-
-/**
- * Description of a plugin template
- */
-struct FPluginTemplateDescription
-{
-	/** Name of this template in the GUI */
-	FText Name;
-
-	/** Description of this template in the GUI */
-	FText Description;
-
-	/** Name of the directory containing template files */
-	FString OnDiskPath;
-
-	/** Brush resource for the image that is dynamically loaded */
-	TSharedPtr< FSlateDynamicImageBrush > PluginIconDynamicImageBrush;
-
-	/** Can the plugin contain content? */
-	bool bCanContainContent;
-
-	/** What is the expected ModuleDescriptor type for this plugin?*/
-	EHostType::Type ModuleDescriptorType;
-
-	/** What is the expected Loading Phase for this plugin? */
-	ELoadingPhase::Type LoadingPhase;
-
-	/** Constructor */
-	FPluginTemplateDescription(FText InName, FText InDescription, FString InOnDiskPath, bool InCanContainContent, EHostType::Type InModuleDescriptorType, ELoadingPhase::Type InLoadingPhase = ELoadingPhase::Default)
-		: Name(InName)
-		, Description(InDescription)
-		, OnDiskPath(InOnDiskPath)
-		, bCanContainContent(InCanContainContent)
-		, ModuleDescriptorType(InModuleDescriptorType)
-		, LoadingPhase(InLoadingPhase)
-	{
-	}
-};
+#include "Features/IPluginsEditorFeature.h"
 
 class IPluginWizardDefinition : public TSharedFromThis<IPluginWizardDefinition>
 {
@@ -60,29 +19,20 @@ public:
 	virtual const TArray<TSharedRef<FPluginTemplateDescription>>& GetTemplatesSource() const = 0;
 
 	/** Changes the selection internally to match the supplied items */
-	virtual void OnTemplateSelectionChanged(TArray<TSharedRef<FPluginTemplateDescription>> InSelectedItems, ESelectInfo::Type SelectInfo) = 0;
+	virtual void OnTemplateSelectionChanged(TSharedPtr<FPluginTemplateDescription> InSelectedItem, ESelectInfo::Type SelectInfo) = 0;
 
 	/** Returns true if the definition has a valid template selection */
 	virtual bool HasValidTemplateSelection() const = 0;
 
-	/** Gets the list selection mode for this definition */
-	virtual ESelectionMode::Type GetSelectionMode() const = 0;
-
-	/** Returns the currently selected templates */
-	virtual TArray<TSharedPtr<FPluginTemplateDescription>> GetSelectedTemplates() const = 0;
+	/** Returns the currently selected template */
+	virtual TSharedPtr<FPluginTemplateDescription> GetSelectedTemplate() const = 0;
 
 	/** Clears the template selection */
 	virtual void ClearTemplateSelection() = 0;
 
-	/** Returns true if this definition allows for the creation of engine plugins */
-	virtual bool AllowsEnginePlugins() const = 0;
-
 	/** Returns true if the wizard for this definition can show on startup */
 	virtual bool CanShowOnStartup() const = 0;
 
-	/** Returns true if the selected template can contain content */
-	virtual bool CanContainContent() const = 0;
-
 	/** Returns true if the selected template will generate code */
 	virtual bool HasModules() const = 0;
 
diff --git a/Engine/Source/Editor/UnrealEd/Public/Features/IPluginsEditorFeature.h b/Engine/Source/Editor/UnrealEd/Public/Features/IPluginsEditorFeature.h
--- a/Engine/Source/Editor/UnrealEd/Public/Features/IPluginsEditorFeature.h
+++ b/Engine/Source/Editor/UnrealEd/Public/Features/IPluginsEditorFeature.h
@@ -4,7 +4,72 @@
 
 #include "CoreMinimal.h"
 #include "Features/IModularFeature.h"
+#include "ModuleDescriptor.h"
 
+struct FSlateDynamicImageBrush;
+class IPlugin;
+
+/**
+ * Description of a plugin template
+ */
+struct FPluginTemplateDescription
+{
+	/** Name of this template in the GUI */
+	FText Name;
+
+	/** Description of this template in the GUI */
+	FText Description;
+
+	/** Path to the directory containing template files */
+	FString OnDiskPath;
+
+	/** Brush resource for the image that is dynamically loaded */
+	TSharedPtr< FSlateDynamicImageBrush > PluginIconDynamicImageBrush;
+
+	/** Sorting priority (higer values go to the top of the list) */
+	int32 SortPriority = 0;
+
+	/** Can the plugin contain content? */
+	bool bCanContainContent = false;
+
+	/** Can the plugin be in the engine folder? */
+	bool bCanBePlacedInEngine = true;
+
+	/** What is the expected ModuleDescriptor type for this plugin? */
+	EHostType::Type ModuleDescriptorType;
+
+	/** What is the expected Loading Phase for this plugin? */
+	ELoadingPhase::Type LoadingPhase;
+
+	/** Called after the plugin has been created */
+	virtual void OnPluginCreated(TSharedPtr<IPlugin> NewPlugin) {}
+
+	/** Called to perform *additional* path validation when the path is modified (the bCanBePlacedInEngine validation will have already occurred and passed by this point) */
+	virtual bool ValidatePathForPlugin(const FString& ProposedAbsolutePluginPath, FText& OutErrorMessage)
+	{
+		OutErrorMessage = FText::GetEmpty();
+		return true;
+	}
+
+	/** Called to enforce any restrictions this template has on paths when it is first selected (so it doesn't generate an error unnecessarily) */
+	virtual void UpdatePathWhenTemplateSelected(FString& InOutPath) {}
+
+	/** Called to change away from special folders if needed */
+	virtual void UpdatePathWhenTemplateUnselected(FString& InOutPath) {}
+
+	/** Constructor */
+	FPluginTemplateDescription(FText InName, FText InDescription, FString InOnDiskPath, bool InCanContainContent, EHostType::Type InModuleDescriptorType, ELoadingPhase::Type InLoadingPhase = ELoadingPhase::Default)
+		: Name(InName)
+		, Description(InDescription)
+		, OnDiskPath(InOnDiskPath)
+		, bCanContainContent(InCanContainContent)
+		, ModuleDescriptorType(InModuleDescriptorType)
+		, LoadingPhase(InLoadingPhase)
+	{
+	}
+
+	virtual ~FPluginTemplateDescription() {}
+};
 
 /**
  * Feature interface for a Plugins management UI
@@ -13,6 +78,14 @@ class IPluginsEditorFeature : public IModularFeature
 {
 
 public:
+	/**
+	  * Registers the specified plugin template with the New Plugin wizard
+	  */
+	virtual void RegisterPluginTemplate(TSharedRef<FPluginTemplateDescription> Template) = 0;
 
+	/**
+	  * Unregisters the specified plugin template from the New Plugin wizard
+	  */
+	virtual void UnregisterPluginTemplate(TSharedRef<FPluginTemplateDescription> Template) = 0;
 };
 
