Fix "New Plugins Available" notifications not disappearing after some amount of time

diff --git a/Engine/Plugins/Editor/PluginBrowser/Source/PluginBrowser/Private/PluginBrowserModule.cpp b/Engine/Plugins/Editor/PluginBrowser/Source/PluginBrowser/Private/PluginBrowserModule.cpp
--- a/Engine/Plugins/Editor/PluginBrowser/Source/PluginBrowser/Private/PluginBrowserModule.cpp
+++ b/Engine/Plugins/Editor/PluginBrowser/Source/PluginBrowser/Private/PluginBrowserModule.cpp
@@ -175,11 +175,11 @@ void FPluginBrowserModule::OnMainFrameLoaded(TSharedPtr<SWindow> InRootWindow, b
 	if(!bIsNewProjectWindow && NewlyInstalledPlugins.Num() > 0 && !PluginBrowserTab.IsValid())
 	{
 		FNotificationInfo Info(LOCTEXT("NewPluginsPopupTitle", "New plugins are available"));
-		Info.bFireAndForget = false;
+		Info.bFireAndForget = true;
 		Info.bUseLargeFont = true;
 		Info.bUseThrobber = false;
-		Info.FadeOutDuration = 0.0f;
-		Info.ExpireDuration = 0.0f;
+		// Note this time is large to add padding because this notification is added before the editor UI is visible. 
+		Info.ExpireDuration = 25.0f;
 		Info.ButtonDetails.Add(FNotificationButtonInfo(LOCTEXT("NewPluginsPopupSettings", "Manage Plugins..."), LOCTEXT("NewPluginsPopupSettingsTT", "Open the plugin browser to enable plugins"), FSimpleDelegate::CreateRaw(this, &FPluginBrowserModule::OnNewPluginsPopupSettingsClicked)));
 		Info.ButtonDetails.Add(FNotificationButtonInfo(LOCTEXT("NewPluginsPopupDismiss", "Dismiss"), LOCTEXT("NewPluginsPopupDismissTT", "Dismiss this notification"), FSimpleDelegate::CreateRaw(this, &FPluginBrowserModule::OnNewPluginsPopupDismissClicked)));
 
@@ -192,12 +192,16 @@ void FPluginBrowserModule::OnNewPluginsPopupSettingsClicked()
 {
 	FGlobalTabmanager::Get()->TryInvokeTab(PluginsEditorTabName);
 	NewPluginsNotification.Pin()->SetCompletionState(SNotificationItem::CS_None);
+	NewPluginsNotification.Pin()->SetExpireDuration(0.0f);
+	NewPluginsNotification.Pin()->SetFadeOutDuration(0.0f);
 	NewPluginsNotification.Pin()->ExpireAndFadeout();
 }
 
 void FPluginBrowserModule::OnNewPluginsPopupDismissClicked()
 {
 	NewPluginsNotification.Pin()->SetCompletionState(SNotificationItem::CS_None);
+	NewPluginsNotification.Pin()->SetExpireDuration(0.0f);
+	NewPluginsNotification.Pin()->SetFadeOutDuration(0.0f);
 	NewPluginsNotification.Pin()->ExpireAndFadeout();
 	UpdatePreviousInstalledPlugins();
 }
