Fixed non-unity compile errors in PluginDescriptorEditor

diff --git a/Engine/Plugins/Editor/PluginBrowser/Source/PluginBrowser/Private/PluginDescriptorEditor.cpp b/Engine/Plugins/Editor/PluginBrowser/Source/PluginBrowser/Private/PluginDescriptorEditor.cpp
--- a/Engine/Plugins/Editor/PluginBrowser/Source/PluginBrowser/Private/PluginDescriptorEditor.cpp
+++ b/Engine/Plugins/Editor/PluginBrowser/Source/PluginBrowser/Private/PluginDescriptorEditor.cpp
@@ -7,10 +7,9 @@
 #include "Modules/ModuleManager.h"
 #include "Widgets/SWindow.h"
 #include "Framework/Application/SlateApplication.h"
-// #include "Widgets/Images/SImage.h"
-// #include "Widgets/Text/STextBlock.h"
-// #include "Widgets/Input/SButton.h"
-// #include "Widgets/Input/SCheckBox.h"
+#include "Widgets/Layout/SBorder.h"
+#include "Widgets/Text/STextBlock.h"
+#include "Widgets/Input/SButton.h"
 #include "EditorStyleSet.h"
 #include "ISourceControlOperation.h"
 #include "SourceControlOperations.h"
@@ -18,18 +17,10 @@
 #include "ISourceControlModule.h"
 #include "PluginDescriptor.h"
 #include "Interfaces/IPluginManager.h"
-// #include "SPluginBrowser.h"
 #include "PluginStyle.h"
-// #include "GameProjectGenerationModule.h"
-// #include "IDetailsView.h"
-// #include "Widgets/Input/SHyperlink.h"
 #include "PluginMetadataObject.h"
 #include "Interfaces/IProjectManager.h"
-// #include "PluginBrowserModule.h"
 #include "PropertyEditorModule.h"
-// #include "IUATHelperModule.h"
-// #include "DesktopPlatformModule.h"
-// #include "Widgets/Layout/SSpacer.h"
 
 #define LOCTEXT_NAMESPACE "PluginListTile"
 
diff --git a/Engine/Plugins/Editor/PluginBrowser/Source/PluginBrowser/Private/PluginDescriptorEditor.h b/Engine/Plugins/Editor/PluginBrowser/Source/PluginBrowser/Private/PluginDescriptorEditor.h
--- a/Engine/Plugins/Editor/PluginBrowser/Source/PluginBrowser/Private/PluginDescriptorEditor.h
+++ b/Engine/Plugins/Editor/PluginBrowser/Source/PluginBrowser/Private/PluginDescriptorEditor.h
@@ -3,6 +3,7 @@
 #pragma once
 
 #include "CoreMinimal.h"
+#include "Input/Reply.h"
 
 class UPluginMetadataObject;
 class IPlugin;
