Provide a drop-down for editing plugin dependencies (using GetOptions to offer other discovered plugins that don't *directly* depend on the plugin being edited)

diff --git a/Engine/Plugins/Editor/PluginBrowser/Source/PluginBrowser/Private/PluginMetadataObject.cpp b/Engine/Plugins/Editor/PluginBrowser/Source/PluginBrowser/Private/PluginMetadataObject.cpp
--- a/Engine/Plugins/Editor/PluginBrowser/Source/PluginBrowser/Private/PluginMetadataObject.cpp
+++ b/Engine/Plugins/Editor/PluginBrowser/Source/PluginBrowser/Private/PluginMetadataObject.cpp
@@ -105,6 +105,41 @@ void UPluginMetadataObject::CopyIntoDescriptor(FPluginDescriptor& OutDescriptor)
 	}
 }
 
+TArray<FString> UPluginMetadataObject::GetAvailablePluginDependencies() const
+{
+	TArray<TSharedRef<IPlugin>> AllPlugins = IPluginManager::Get().GetDiscoveredPlugins();
+	TArray<FString> Result;
+	Result.Reserve(AllPlugins.Num());
+
+	const FString MyName = SourcePlugin.Pin()->GetName();
+
+	// This lists all plugins that don't directly depend on the plugin being edited (so a multi-hop chain would still cause a problem)
+	for (const TSharedRef<IPlugin>& Plugin : AllPlugins)
+	{
+		if (Plugin == SourcePlugin)
+		{
+			continue;
+		}
+
+		bool bDependsOnMe = false;
+		for (const FPluginReferenceDescriptor& Dependency : Plugin->GetDescriptor().Plugins)
+		{
+			if (Dependency.Name == MyName)
+			{
+				bDependsOnMe = true;
+				break;
+			}
+		}
+
+		if (!bDependsOnMe)
+		{
+			Result.Add(Plugin->GetName());
+		}
+	}
+
+	return Result;
+}
+
 /////////////////////////////////////////////////////////////////////////////////////////////////////////
 
 TSharedRef<IDetailCustomization> FPluginMetadataCustomization::MakeInstance()
diff --git a/Engine/Plugins/Editor/PluginBrowser/Source/PluginBrowser/Private/PluginMetadataObject.h b/Engine/Plugins/Editor/PluginBrowser/Source/PluginBrowser/Private/PluginMetadataObject.h
--- a/Engine/Plugins/Editor/PluginBrowser/Source/PluginBrowser/Private/PluginMetadataObject.h
+++ b/Engine/Plugins/Editor/PluginBrowser/Source/PluginBrowser/Private/PluginMetadataObject.h
@@ -24,7 +24,7 @@ public:
 	GENERATED_BODY()
 
 	/** Name of the dependency plugin */
-	UPROPERTY(EditAnywhere, Category = "Plugin Reference")
+	UPROPERTY(EditAnywhere, Category = "Plugin Reference", meta=(GetOptions=GetAvailablePluginDependencies))
 	FString Name;
 
 	/** Whether the dependency plugin is optional meaning it will be silently ignored if not present */
@@ -113,7 +113,7 @@ public:
 	bool bIsBetaVersion;
 
 	/** Plugins used by this plugin */
-	UPROPERTY(EditAnywhere, Category = Dependencies)
+	UPROPERTY(EditAnywhere, Category = Dependencies, meta=(TitleProperty=Name))
 	TArray<FPluginReferenceMetadata> Plugins;
 
 	/** Plugin this proxy object was constructed from */
@@ -131,6 +131,9 @@ public:
 	 * Copy the metadata fields into a plugin descriptor.
 	 */
 	void CopyIntoDescriptor(FPluginDescriptor& OutDescriptor) const;
+
+	UFUNCTION()
+	TArray<FString> GetAvailablePluginDependencies() const;
 };
 
 /**
