Added a scoped slow task when checking out .uplugin files

diff --git a/Engine/Plugins/Editor/PluginBrowser/Source/PluginBrowser/Private/PluginDescriptorEditor.cpp b/Engine/Plugins/Editor/PluginBrowser/Source/PluginBrowser/Private/PluginDescriptorEditor.cpp
--- a/Engine/Plugins/Editor/PluginBrowser/Source/PluginBrowser/Private/PluginDescriptorEditor.cpp
+++ b/Engine/Plugins/Editor/PluginBrowser/Source/PluginBrowser/Private/PluginDescriptorEditor.cpp
@@ -21,6 +21,7 @@
 #include "PluginMetadataObject.h"
 #include "Interfaces/IProjectManager.h"
 #include "PropertyEditorModule.h"
+#include "Misc/ScopedSlowTask.h"
 
 #define LOCTEXT_NAMESPACE "PluginListTile"
 
@@ -116,6 +117,8 @@ FReply FPluginDescriptorEditor::OnEditPluginFinished(UPluginMetadataObject* Meta
 		ISourceControlModule& SourceControlModule = ISourceControlModule::Get();
 		if(SourceControlModule.IsEnabled())
 		{
+			FScopedSlowTask SlowTask(0, LOCTEXT("CheckOutPlugin", "Checking out plugin..."));
+
 			ISourceControlProvider& SourceControlProvider = SourceControlModule.GetProvider();
 			TSharedPtr<ISourceControlState, ESPMode::ThreadSafe> SourceControlState = SourceControlProvider.GetState(DescriptorFileName, EStateCacheUsage::ForceUpdate);
 			if(SourceControlState.IsValid() && SourceControlState->CanCheckout())
