Updating wording and adding help text to the beta and experimental plugin warnings in the plugin manager.

diff --git a/Engine/Plugins/Editor/PluginBrowser/Source/PluginBrowser/Private/SPluginTile.cpp b/Engine/Plugins/Editor/PluginBrowser/Source/PluginBrowser/Private/SPluginTile.cpp
--- a/Engine/Plugins/Editor/PluginBrowser/Source/PluginBrowser/Private/SPluginTile.cpp
+++ b/Engine/Plugins/Editor/PluginBrowser/Source/PluginBrowser/Private/SPluginTile.cpp
@@ -326,6 +326,7 @@ void SPluginTile::RecreateWidgets()
 																SNew(STextBlock)
 																	.TextStyle(FPluginStyle::Get(), "PluginTile.BetaText")
 																	.Text(PluginDescriptor.bIsBetaVersion ? LOCTEXT("PluginBetaVersionText", "BETA") : LOCTEXT("PluginExperimentalVersionText", "EXPERIMENTAL"))
+																	.ToolTipText(this, &SPluginTile::GetBetaOrExperimentalHelpText)
 															]
 													]
 
@@ -490,7 +491,15 @@ void SPluginTile::OnEnablePluginCheckboxChanged(ECheckBoxState NewCheckedState)
 		// If this is plugin is marked as beta, make sure the user is aware before enabling it.
 		if (PluginDescriptor.bIsBetaVersion)
 		{
-			FText WarningMessage = FText::Format(LOCTEXT("Warning_EnablingBetaPlugin", "Plugin '{0}' is a beta version and might be unstable or removed without notice. Please use with caution. Are you sure you want to enable the plugin?"), GetPluginNameText());
+			FText WarningMessage = FText::Format(LOCTEXT("Warning_EnablingBetaPlugin", "Plugin '{0}' is a beta version. {1} Are you sure you want to enable the plugin?"), GetPluginNameText(), GetBetaOrExperimentalHelpText());
+			if (EAppReturnType::No == FMessageDialog::Open(EAppMsgType::YesNo, WarningMessage))
+			{
+				return;
+			}
+		}
+		else if (PluginDescriptor.bIsExperimentalVersion)
+		{
+			FText WarningMessage = FText::Format(LOCTEXT("Warning_EnablingExperimentalPlugin", "Plugin '{0}' is an experimental version. {1} Are you sure you want to enable the plugin?"), GetPluginNameText(), GetBetaOrExperimentalHelpText());
 			if (EAppReturnType::No == FMessageDialog::Open(EAppMsgType::YesNo, WarningMessage))
 			{
 				return;
@@ -627,4 +636,24 @@ void SPluginTile::OnPackagePlugin()
 		LOCTEXT("PackagePluginTaskShortName", "Package Plugin Task"), FEditorStyle::GetBrush(TEXT("MainFrame.CookContent")));
 }
 
+FText SPluginTile::GetBetaOrExperimentalHelpText() const
+{
+	if (!Plugin)
+	{
+		return FText();
+	}
+
+	const FPluginDescriptor& PluginDescriptor = Plugin->GetDescriptor();
+	if (PluginDescriptor.bIsBetaVersion)
+	{
+		return LOCTEXT("Description_BetaPlugin", "Epic recommends using caution when shipping projects with beta plugins. Beta plugins support backwards compatibility for assets and APIs, but performance, stability, and platform support may not be shipping quality.");
+	}
+	if (PluginDescriptor.bIsExperimentalVersion)
+	{
+		return LOCTEXT("Description_ExperimentalPlugin", "Epic does not recommend shipping projects with experimental plugins. APIs, features, and the plugin itself are subject to change or be removed without notice.");
+	}
+
+	return FText();
+}
+
 #undef LOCTEXT_NAMESPACE
diff --git a/Engine/Plugins/Editor/PluginBrowser/Source/PluginBrowser/Private/SPluginTile.h b/Engine/Plugins/Editor/PluginBrowser/Source/PluginBrowser/Private/SPluginTile.h
--- a/Engine/Plugins/Editor/PluginBrowser/Source/PluginBrowser/Private/SPluginTile.h
+++ b/Engine/Plugins/Editor/PluginBrowser/Source/PluginBrowser/Private/SPluginTile.h
@@ -54,6 +54,9 @@ private:
 	/** Called when the 'package' hyperlink is clicked */
 	void OnPackagePlugin();
 
+	/** Called when the beta or experimental flags are hovered over, or when a beta or experimental plugin is enabled */
+	FText GetBetaOrExperimentalHelpText() const;
+
 private:
 
 	/** The item we're representing the in tree */
