Editor: Fixed a typo in the Plugin UI about editing being disabled.

diff --git a/Engine/Plugins/Editor/PluginBrowser/Source/PluginBrowser/Private/SPluginTile.cpp b/Engine/Plugins/Editor/PluginBrowser/Source/PluginBrowser/Private/SPluginTile.cpp
--- a/Engine/Plugins/Editor/PluginBrowser/Source/PluginBrowser/Private/SPluginTile.cpp
+++ b/Engine/Plugins/Editor/PluginBrowser/Source/PluginBrowser/Private/SPluginTile.cpp
@@ -326,7 +326,7 @@ void SPluginTile::RecreateWidgets()
 									.IsEnabled(this, &SPluginTile::CanModifyPlugins)
 									.ToolTipText(CanModifyPlugins() ?
 										LOCTEXT("EnableDisableButtonToolTip", "Toggles whether this plugin is enabled for your current project.  You may need to restart the program for this change to take effect.")
-										: LOCTEXT("NonEditableButtonToolTip", "Editting plugin enabled/disabled state from the Plugin Browser has been disabled for this project."))
+										: LOCTEXT("NonEditableButtonToolTip", "Editing plugin enabled/disabled state from the Plugin Browser has been disabled for this project."))
 							]
 						// Thumbnail image
 						+ SHorizontalBox::Slot()
