Tagged PluginBrowser allocations

diff --git a/Engine/Plugins/Editor/PluginBrowser/Source/PluginBrowser/Private/PluginBrowserModule.cpp b/Engine/Plugins/Editor/PluginBrowser/Source/PluginBrowser/Private/PluginBrowserModule.cpp
--- a/Engine/Plugins/Editor/PluginBrowser/Source/PluginBrowser/Private/PluginBrowserModule.cpp
+++ b/Engine/Plugins/Editor/PluginBrowser/Source/PluginBrowser/Private/PluginBrowserModule.cpp
@@ -17,9 +17,12 @@
 #include "Misc/ConfigCacheIni.h"
 #include "Interfaces/IPluginManager.h"
 #include "PluginDescriptorEditor.h"
+#include "HAL/LowLevelMemTracker.h"
 
 #define LOCTEXT_NAMESPACE "PluginsEditor"
 
+LLM_DEFINE_TAG(PluginBrowser);
+
 IMPLEMENT_MODULE( FPluginBrowserModule, PluginBrowser )
 
 const FName FPluginBrowserModule::PluginsEditorTabName( TEXT( "PluginsEditor" ) );
@@ -27,6 +30,8 @@ const FName FPluginBrowserModule::PluginCreatorTabName( TEXT( "PluginCreator" )
 
 void FPluginBrowserModule::StartupModule()
 {
+	LLM_SCOPE_BYTAG(PluginBrowser);
+
 	FPluginStyle::Initialize();
 
 	// Register ourselves as an editor feature
@@ -101,6 +106,7 @@ void FPluginBrowserModule::ShutdownModule()
 
 void FPluginBrowserModule::RegisterPluginTemplate(TSharedRef<FPluginTemplateDescription> Template)
 {
+	LLM_SCOPE_BYTAG(PluginBrowser);
 	AddedPluginTemplates.Add(Template);
 }
 
@@ -111,6 +117,7 @@ void FPluginBrowserModule::UnregisterPluginTemplate(TSharedRef<FPluginTemplateDe
 
 FPluginEditorExtensionHandle FPluginBrowserModule::RegisterPluginEditorExtension(FOnPluginBeingEdited Extension)
 {
+	LLM_SCOPE_BYTAG(PluginBrowser);
 	++EditorExtensionCounter;
 	FPluginEditorExtensionHandle Result = EditorExtensionCounter;
 	CustomizePluginEditingDelegates.Add(MakeTuple(Extension, Result));
