Make vendor links in the Plugin Browser use new Open In Web Browser icon, and update other Open In icons to match.

diff --git a/Engine/Plugins/Editor/PluginBrowser/Source/PluginBrowser/Private/PluginStyle.cpp b/Engine/Plugins/Editor/PluginBrowser/Source/PluginBrowser/Private/PluginStyle.cpp
--- a/Engine/Plugins/Editor/PluginBrowser/Source/PluginBrowser/Private/PluginStyle.cpp
+++ b/Engine/Plugins/Editor/PluginBrowser/Source/PluginBrowser/Private/PluginStyle.cpp
@@ -60,7 +60,6 @@ void FPluginStyle::Initialize()
 
 	StyleSet->Set( "Plugins.TabIcon", new IMAGE_BRUSH_SVG( "Plugins", Icon16x16 ) );
 	StyleSet->Set( "Plugins.BreadcrumbArrow", new IMAGE_BRUSH( "SmallArrowRight", Icon10x10 ) );
-	StyleSet->Set( "Plugins.Documentation", new CORE_IMAGE_BRUSH_SVG("../Editor/Slate/Starship/Common/Documentation", Icon16x16));
 	StyleSet->Set( "Plugins.ListBorder", new FSlateRoundedBoxBrush(FStyleColors::Recessed, 4.0f));
 	StyleSet->Set( "Plugins.RestartWarningBorder", new FSlateRoundedBoxBrush(FStyleColors::Panel , 5.0f, FStyleColors::Warning, 1.0f));
 
diff --git a/Engine/Plugins/Editor/PluginBrowser/Source/PluginBrowser/Private/SPluginTile.cpp b/Engine/Plugins/Editor/PluginBrowser/Source/PluginBrowser/Private/SPluginTile.cpp
--- a/Engine/Plugins/Editor/PluginBrowser/Source/PluginBrowser/Private/SPluginTile.cpp
+++ b/Engine/Plugins/Editor/PluginBrowser/Source/PluginBrowser/Private/SPluginTile.cpp
@@ -161,15 +161,6 @@ void SPluginTile::RecreateWidgets()
 		{
 			CreatedByWidget = SNew(SHorizontalBox)
 
-			+ SHorizontalBox::Slot()
-				.AutoWidth()
-				.VAlign(VAlign_Center)
-				[
-					SNew(SImage)
-						.ColorAndOpacity(FSlateColor::UseForeground())
-						.Image(FAppStyle::GetBrush("ContentBrowser.AssetTreeFolderDeveloper"))
-				]
-
 			+ SHorizontalBox::Slot()
 				.AutoWidth()
 				.VAlign(VAlign_Center)
@@ -193,6 +184,16 @@ void SPluginTile::RecreateWidgets()
 						.ToolTipText(FText::Format(LOCTEXT("NavigateToCreatedByURL", "Visit the vendor's web site ({0})"), FText::FromString(CreatedByURL)))
 						.OnNavigate_Lambda([=]() { FPlatformProcess::LaunchURL(*CreatedByURL, nullptr, nullptr); })
 						.Style(FAppStyle::Get(), "HoverOnlyHyperlink")
+				]
+
+			+ SHorizontalBox::Slot()
+				.AutoWidth()
+				.VAlign(VAlign_Center)
+				.Padding(4.0f, 0.0f, 0.0f, 0.0f)
+				[
+					SNew(SImage)
+					.ColorAndOpacity(FSlateColor::UseForeground())
+					.Image(FAppStyle::Get().GetBrush("Icons.OpenInBrowser"))
 				];
 		}
 	}
@@ -472,7 +473,7 @@ void SPluginTile::RecreateWidgets()
 												.AutoWidth()
 												.HAlign(HAlign_Right)
 												.VAlign(VAlign_Top)
-												.Padding(0.f, 0.f, PaddingAmount + 14.f, 0.f)
+												.Padding(PaddingAmount + 14.f, 0.f, PaddingAmount + 14.f, 0.f)
 												[
 													CreatedByWidget.ToSharedRef()
 												]
diff --git a/Engine/Source/Editor/EditorStyle/Private/StarshipStyle.cpp b/Engine/Source/Editor/EditorStyle/Private/StarshipStyle.cpp
--- a/Engine/Source/Editor/EditorStyle/Private/StarshipStyle.cpp
+++ b/Engine/Source/Editor/EditorStyle/Private/StarshipStyle.cpp
@@ -422,6 +422,7 @@ void FStarshipEditorStyle::FStyle::SetupGeneralStyles()
 		Set("Icons.SkeletalMesh", new IMAGE_BRUSH_SVG("Starship/Common/SkeletalMesh", Icon16x16));
 		Set("Icons.OpenInExternalEditor", new IMAGE_BRUSH_SVG("Starship/Common/OpenInExternalEditor", Icon16x16));
 		Set("Icons.OpenSourceLocation", new IMAGE_BRUSH_SVG("Starship/Common/OpenSourceLocation", Icon16x16));
+		Set("Icons.OpenInBrowser", new IMAGE_BRUSH_SVG("Starship/Common/WebBrowser", Icon16x16));
 		Set("Icons.Find", new IMAGE_BRUSH_SVG("Starship/Common/Find", Icon16x16));
 		Set("Icons.Validate", new IMAGE_BRUSH_SVG("Starship/Common/validate", Icon16x16));
 		Set("Icons.Pinned", new IMAGE_BRUSH_SVG("Starship/Common/Pinned", Icon16x16));
