[PluginBrowser]
* Fixed non unity compile error

diff --git a/Engine/Plugins/Editor/PluginBrowser/Source/PluginBrowser/Private/SPluginBrowser.cpp b/Engine/Plugins/Editor/PluginBrowser/Source/PluginBrowser/Private/SPluginBrowser.cpp
--- a/Engine/Plugins/Editor/PluginBrowser/Source/PluginBrowser/Private/SPluginBrowser.cpp
+++ b/Engine/Plugins/Editor/PluginBrowser/Source/PluginBrowser/Private/SPluginBrowser.cpp
@@ -22,6 +22,7 @@
 #include "Framework/MultiBox/MultiBoxBuilder.h"
 #include "Styling/StyleColors.h"
 #include "SWarningOrErrorBox.h"
+#include "Widgets/Views/SListView.h"
 
 #define LOCTEXT_NAMESPACE "PluginsEditor"
 
