Exposing the plugin descriptor property EditorVirtualPath in the plugin properties editor window

diff --git a/Engine/Plugins/Editor/PluginBrowser/Source/PluginBrowser/Private/PluginMetadataObject.cpp b/Engine/Plugins/Editor/PluginBrowser/Source/PluginBrowser/Private/PluginMetadataObject.cpp
--- a/Engine/Plugins/Editor/PluginBrowser/Source/PluginBrowser/Private/PluginMetadataObject.cpp
+++ b/Engine/Plugins/Editor/PluginBrowser/Source/PluginBrowser/Private/PluginMetadataObject.cpp
@@ -46,6 +46,7 @@ void UPluginMetadataObject::PopulateFromPlugin(TSharedPtr<IPlugin> InPlugin)
 	DocsURL = InDescriptor.DocsURL;
 	MarketplaceURL = InDescriptor.MarketplaceURL;
 	SupportURL = InDescriptor.SupportURL;
+	EditorCustomVirtualPath = InDescriptor.EditorCustomVirtualPath;
 	bCanContainContent = InDescriptor.bCanContainContent;
 	bIsBetaVersion = InDescriptor.bIsBetaVersion;
 	bIsEnabledByDefault = (InDescriptor.EnabledByDefault == EPluginEnabledByDefault::Enabled);
@@ -70,6 +71,7 @@ void UPluginMetadataObject::CopyIntoDescriptor(FPluginDescriptor& OutDescriptor)
 	OutDescriptor.CreatedByURL = CreatedByURL;
 	OutDescriptor.DocsURL = DocsURL;
 	OutDescriptor.MarketplaceURL = MarketplaceURL;
+	OutDescriptor.EditorCustomVirtualPath = EditorCustomVirtualPath;
 	OutDescriptor.SupportURL = SupportURL;
 	OutDescriptor.bCanContainContent = bCanContainContent;
 	OutDescriptor.bIsBetaVersion = bIsBetaVersion;
diff --git a/Engine/Plugins/Editor/PluginBrowser/Source/PluginBrowser/Private/PluginMetadataObject.h b/Engine/Plugins/Editor/PluginBrowser/Source/PluginBrowser/Private/PluginMetadataObject.h
--- a/Engine/Plugins/Editor/PluginBrowser/Source/PluginBrowser/Private/PluginMetadataObject.h
+++ b/Engine/Plugins/Editor/PluginBrowser/Source/PluginBrowser/Private/PluginMetadataObject.h
@@ -101,6 +101,10 @@ public:
 	UPROPERTY(EditAnywhere, Category = Details)
 	FString SupportURL;
 
+	/** Optional custom virtual path to display in editor to better organize. Inserted just before this plugin's directory in the path: /All/Plugins/EditorCustomVirtualPath/PluginName */
+	UPROPERTY(EditAnywhere, Category = Details)
+	FString EditorCustomVirtualPath;
+
 	/** Can this plugin contain content? */
 	UPROPERTY(EditAnywhere, Category = Details)
 	bool bCanContainContent;
