Fixing a bug that occurs when creating a new plugin from a template containing .uasset files. It's possible for the plugin to fail to register if the post copy rename is triggered while the asset registry scan is still pending. The calling code assumes that the rename is synchronous however it can defer the operation and return. This puts the plugin in a bad state.

Showing the asset discovery dialog and preventing the plugin creation until the asset discovery is complete.

diff --git a/Engine/Plugins/Editor/PluginBrowser/Source/PluginBrowser/Private/SNewPluginWizard.cpp b/Engine/Plugins/Editor/PluginBrowser/Source/PluginBrowser/Private/SNewPluginWizard.cpp
--- a/Engine/Plugins/Editor/PluginBrowser/Source/PluginBrowser/Private/SNewPluginWizard.cpp
+++ b/Engine/Plugins/Editor/PluginBrowser/Source/PluginBrowser/Private/SNewPluginWizard.cpp
@@ -1,12 +1,14 @@
 // Copyright Epic Games, Inc. All Rights Reserved.
 
 #include "SNewPluginWizard.h"
+#include "AssetToolsModule.h"
 #include "Brushes/SlateDynamicImageBrush.h"
 #include "Misc/Paths.h"
 #include "DetailsViewArgs.h"
 #include "Misc/MessageDialog.h"
 #include "Framework/Views/TableViewMetadata.h"
 #include "HAL/FileManager.h"
+#include "IAssetTools.h"
 #include "IDesktopPlatform.h"
 #include "Misc/App.h"
 #include "Framework/Application/SlateApplication.h"
@@ -51,6 +53,7 @@ SNewPluginWizard::SNewPluginWizard()
 	: bIsPluginPathValid(false)
 	, bIsPluginNameValid(false)
 	, bIsSelectedPathInEngine(false)
+	, bIsWaitingForAssetDiscoveryToFinish(false)
 {
 	AbsoluteGamePluginPath = IFileManager::Get().ConvertToAbsolutePathForExternalAppForWrite(*FPaths::ProjectPluginsDir());
 	FPaths::MakePlatformFilename(AbsoluteGamePluginPath);
@@ -573,7 +576,7 @@ void SNewPluginWizard::ValidateFullPluginPath()
 
 bool SNewPluginWizard::CanCreatePlugin() const
 {
-	return bIsPluginPathValid && bIsPluginNameValid && PluginWizardDefinition->HasValidTemplateSelection();
+	return bIsPluginPathValid && bIsPluginNameValid && PluginWizardDefinition->HasValidTemplateSelection() && !bIsWaitingForAssetDiscoveryToFinish;
 }
 
 FText SNewPluginWizard::GetPluginDestinationPath() const
@@ -624,6 +627,13 @@ FReply SNewPluginWizard::OnCreatePluginClicked()
 		return FReply::Unhandled();
 	}
 
+	FAssetToolsModule& AssetToolsModule = FModuleManager::LoadModuleChecked<FAssetToolsModule>("AssetTools");
+	if (AssetToolsModule.Get().IsDiscoveringAssetsInProgress())
+	{
+		bIsWaitingForAssetDiscoveryToFinish = true;
+		AssetToolsModule.Get().OpenDiscoveringAssetsDialog(IAssetTools::FOnAssetsDiscovered::CreateSP(this, &SNewPluginWizard::DiscoveringAssetsCalback));
+		return FReply::Unhandled();
+	}
 
 	const FString PluginName = PluginNameText.ToString();
 	const bool bHasModules = PluginWizardDefinition->HasModules();
@@ -696,4 +706,9 @@ FReply SNewPluginWizard::OnCreatePluginClicked()
 	}
 }
 
+void SNewPluginWizard::DiscoveringAssetsCalback()
+{
+	bIsWaitingForAssetDiscoveryToFinish = false;
+}
+
 #undef LOCTEXT_NAMESPACE
diff --git a/Engine/Plugins/Editor/PluginBrowser/Source/PluginBrowser/Private/SNewPluginWizard.h b/Engine/Plugins/Editor/PluginBrowser/Source/PluginBrowser/Private/SNewPluginWizard.h
--- a/Engine/Plugins/Editor/PluginBrowser/Source/PluginBrowser/Private/SNewPluginWizard.h
+++ b/Engine/Plugins/Editor/PluginBrowser/Source/PluginBrowser/Private/SNewPluginWizard.h
@@ -110,6 +110,11 @@ private:
 	 */
 	FReply OnCreatePluginClicked();
 
+	/**
+	 * Called when asset discovering completes if the user attempts to create a new plugin while asset discovery is still in progress.
+	 */
+	void DiscoveringAssetsCalback();
+
 	/**
 	 * Generates the expected list view for the plugin wizard, based on the wizard's definition
 	 */
@@ -160,6 +165,9 @@ private:
 	/** Whether we want to create a plugin in the engine folder */
 	bool bIsSelectedPathInEngine;
 
+	/** Whether the asset discovery dialog is open and the user needs to wait on it */
+	bool bIsWaitingForAssetDiscoveryToFinish;
+
 	/** Tab that owns this wizard so that we can ask to close after completion */
 	TWeakPtr<SDockTab> OwnerTab;
 };
diff --git a/Engine/Source/Developer/AssetTools/Private/AssetTools.cpp b/Engine/Source/Developer/AssetTools/Private/AssetTools.cpp
--- a/Engine/Source/Developer/AssetTools/Private/AssetTools.cpp
+++ b/Engine/Source/Developer/AssetTools/Private/AssetTools.cpp
@@ -2441,6 +2441,21 @@ bool UAssetToolsImpl::AdvancedCopyPackages(const FAdvancedCopyParams& CopyParams
 	return bResult;
 }
 
+bool UAssetToolsImpl::IsDiscoveringAssetsInProgress() const
+{
+	FAssetRegistryModule& AssetRegistryModule = FModuleManager::LoadModuleChecked<FAssetRegistryModule>("AssetRegistry");
+	return AssetRegistryModule.Get().IsLoadingAssets();
+}
+
+void UAssetToolsImpl::OpenDiscoveringAssetsDialog(const FOnAssetsDiscovered& InOnAssetsDiscovered)
+{
+	FAssetRegistryModule& AssetRegistryModule = FModuleManager::LoadModuleChecked<FAssetRegistryModule>("AssetRegistry");
+	if (AssetRegistryModule.Get().IsLoadingAssets())
+	{
+		SDiscoveringAssetsDialog::OpenDiscoveringAssetsDialog(InOnAssetsDiscovered);
+	}
+}
+
 bool UAssetToolsImpl::RenameAssets(const TArray<FAssetRenameData>& AssetsAndNames)
 {
 	return AssetRenameManager->RenameAssets(AssetsAndNames);
diff --git a/Engine/Source/Developer/AssetTools/Private/AssetTools.h b/Engine/Source/Developer/AssetTools/Private/AssetTools.h
--- a/Engine/Source/Developer/AssetTools/Private/AssetTools.h
+++ b/Engine/Source/Developer/AssetTools/Private/AssetTools.h
@@ -89,6 +89,8 @@ public:
 	virtual UObject* DuplicateAssetWithDialogAndTitle(const FString& AssetName, const FString& PackagePath, UObject* OriginalObject, FText DialogTitle) override;
 	virtual void SetCreateAssetsAsExternallyReferenceable(bool bValue) override;
 	virtual bool GetCreateAssetsAsExternallyReferenceable() override;
+	virtual bool IsDiscoveringAssetsInProgress() const override;
+	virtual void OpenDiscoveringAssetsDialog(const FOnAssetsDiscovered& InOnAssetsDiscovered) override;
 	virtual bool RenameAssets(const TArray<FAssetRenameData>& AssetsAndNames) override;
 	virtual EAssetRenameResult RenameAssetsWithDialog(const TArray<FAssetRenameData>& AssetsAndNames, bool bAutoCheckout = false) override;
 	virtual void FindSoftReferencesToObject(FSoftObjectPath TargetObject, TArray<UObject*>& ReferencingObjects) override;
diff --git a/Engine/Source/Developer/AssetTools/Public/IAssetTools.h b/Engine/Source/Developer/AssetTools/Public/IAssetTools.h
--- a/Engine/Source/Developer/AssetTools/Public/IAssetTools.h
+++ b/Engine/Source/Developer/AssetTools/Public/IAssetTools.h
@@ -246,6 +246,8 @@ class IAssetTools
 public:
 	ASSETTOOLS_API static IAssetTools& Get();
 
+	DECLARE_DELEGATE(FOnAssetsDiscovered);
+
 	/** Registers an asset type actions object so it can provide information about and actions for asset types. */
 	virtual void RegisterAssetTypeActions(const TSharedRef<IAssetTypeActions>& NewActions) = 0;
 
@@ -354,6 +356,12 @@ public:
 	/** Gets whether assets are being made externally referenceable or not */
 	virtual bool GetCreateAssetsAsExternallyReferenceable() = 0;
 
+	/** Gets whether assets registry is still loading assets or not */
+	virtual bool IsDiscoveringAssetsInProgress() const = 0;
+
+	/** Opens a dialog asking the user to wait while assets are being discovered */
+	virtual void OpenDiscoveringAssetsDialog(const FOnAssetsDiscovered& InOnAssetsDiscovered) = 0;
+
 	/** Renames assets using the specified names. */
 	UFUNCTION(BlueprintCallable, Category = "Editor Scripting | Asset Tools")
 	virtual bool RenameAssets(const TArray<FAssetRenameData>& AssetsAndNames) = 0;
