Added a new 'Sealed' plugin concept that can prevent other plugins from depending on a plugin

diff --git a/Engine/Plugins/Editor/PluginBrowser/Source/PluginBrowser/Private/PluginMetadataObject.cpp b/Engine/Plugins/Editor/PluginBrowser/Source/PluginBrowser/Private/PluginMetadataObject.cpp
--- a/Engine/Plugins/Editor/PluginBrowser/Source/PluginBrowser/Private/PluginMetadataObject.cpp
+++ b/Engine/Plugins/Editor/PluginBrowser/Source/PluginBrowser/Private/PluginMetadataObject.cpp
@@ -51,6 +51,7 @@ void UPluginMetadataObject::PopulateFromPlugin(TSharedPtr<IPlugin> InPlugin)
 	bIsBetaVersion = InDescriptor.bIsBetaVersion;
 	bIsEnabledByDefault = (InDescriptor.EnabledByDefault == EPluginEnabledByDefault::Enabled);
 	bExplicitlyLoaded = InDescriptor.bExplicitlyLoaded;
+	bIsSealed = InDescriptor.bIsSealed;
 
 	Plugins.Reset(InDescriptor.Plugins.Num());
 	for (const FPluginReferenceDescriptor& PluginRefDesc : InDescriptor.Plugins)
@@ -75,6 +76,7 @@ void UPluginMetadataObject::CopyIntoDescriptor(FPluginDescriptor& OutDescriptor)
 	OutDescriptor.SupportURL = SupportURL;
 	OutDescriptor.bCanContainContent = bCanContainContent;
 	OutDescriptor.bIsBetaVersion = bIsBetaVersion;
+	OutDescriptor.bIsSealed = bIsSealed;
 
 	TArray<FPluginReferenceDescriptor> NewPlugins;
 	NewPlugins.Reserve(Plugins.Num());
@@ -122,6 +124,12 @@ TArray<FString> UPluginMetadataObject::GetAvailablePluginDependencies() const
 			continue;
 		}
 
+		// Can't depend on plugins that are marked as sealed.
+		if (Plugin->GetDescriptor().bIsSealed)
+		{
+			continue;
+		}
+
 		bool bDependsOnMe = false;
 		for (const FPluginReferenceDescriptor& Dependency : Plugin->GetDescriptor().Plugins)
 		{
diff --git a/Engine/Plugins/Editor/PluginBrowser/Source/PluginBrowser/Private/PluginMetadataObject.h b/Engine/Plugins/Editor/PluginBrowser/Source/PluginBrowser/Private/PluginMetadataObject.h
--- a/Engine/Plugins/Editor/PluginBrowser/Source/PluginBrowser/Private/PluginMetadataObject.h
+++ b/Engine/Plugins/Editor/PluginBrowser/Source/PluginBrowser/Private/PluginMetadataObject.h
@@ -121,6 +121,10 @@ public:
 	UPROPERTY(VisibleAnywhere, Category = Details)
 	bool bExplicitlyLoaded;
 
+	/** If true, prevents other plugins from depending on this plugin */
+	UPROPERTY(VisibleAnywhere, Category = Details)
+	bool bIsSealed;
+
 	/** Plugins used by this plugin */
 	UPROPERTY(EditAnywhere, Category = Dependencies, meta=(TitleProperty=Name))
 	TArray<FPluginReferenceMetadata> Plugins;
diff --git a/Engine/Source/Runtime/Projects/Private/PluginDescriptor.cpp b/Engine/Source/Runtime/Projects/Private/PluginDescriptor.cpp
--- a/Engine/Source/Runtime/Projects/Private/PluginDescriptor.cpp
+++ b/Engine/Source/Runtime/Projects/Private/PluginDescriptor.cpp
@@ -86,6 +86,7 @@ FPluginDescriptor::FPluginDescriptor()
 	, bInstalled(false)
 	, bRequiresBuildPlatform(false)
 	, bIsHidden(false)
+	, bIsSealed(false)
 	, bExplicitlyLoaded(false)
 	, bHasExplicitPlatforms(false)
 	, bIsPluginExtension(false)
@@ -246,6 +247,7 @@ bool FPluginDescriptor::Read(const FJsonObject& Object, FText* OutFailReason /*=
 	Object.TryGetBoolField(TEXT("Installed"), bInstalled);
 	Object.TryGetBoolField(TEXT("RequiresBuildPlatform"), bRequiresBuildPlatform);
 	Object.TryGetBoolField(TEXT("Hidden"), bIsHidden);
+	Object.TryGetBoolField(TEXT("Sealed"), bIsSealed);
 	Object.TryGetBoolField(TEXT("ExplicitlyLoaded"), bExplicitlyLoaded);
 	Object.TryGetBoolField(TEXT("HasExplicitPlatforms"), bHasExplicitPlatforms);
 
@@ -435,6 +437,15 @@ void FPluginDescriptor::UpdateJson(FJsonObject& JsonObject) const
 		JsonObject.RemoveField(TEXT("Hidden"));
 	}
 
+	if (bIsSealed)
+	{
+		JsonObject.SetBoolField(TEXT("Sealed"), bIsSealed);
+	}
+	else
+	{
+		JsonObject.RemoveField(TEXT("Sealed"));
+	}
+
 	if (bExplicitlyLoaded)
 	{
 		JsonObject.SetBoolField(TEXT("ExplicitlyLoaded"), bExplicitlyLoaded);
diff --git a/Engine/Source/Runtime/Projects/Public/PluginDescriptor.h b/Engine/Source/Runtime/Projects/Public/PluginDescriptor.h
--- a/Engine/Source/Runtime/Projects/Public/PluginDescriptor.h
+++ b/Engine/Source/Runtime/Projects/Public/PluginDescriptor.h
@@ -122,6 +122,9 @@ struct PROJECTS_API FPluginDescriptor
 	/** For auto-generated plugins that should not be listed in the plugin browser for users to disable freely. */
 	bool bIsHidden;
 
+	/** Prevents other plugins from depending on this plugin. */
+	bool bIsSealed;
+
 	/** When true, this plugin's modules will not be loaded automatically nor will it's content be mounted automatically. It will load/mount when explicitly requested and LoadingPhases will be ignored */
 	bool bExplicitlyLoaded;
 
