Changing UPluginMetadataObject::bIsSealed property specifier to EditAnywhere to allow editing from the Plugin Properties window

#rnx

diff --git a/Engine/Plugins/Editor/PluginBrowser/Source/PluginBrowser/Private/PluginMetadataObject.h b/Engine/Plugins/Editor/PluginBrowser/Source/PluginBrowser/Private/PluginMetadataObject.h
--- a/Engine/Plugins/Editor/PluginBrowser/Source/PluginBrowser/Private/PluginMetadataObject.h
+++ b/Engine/Plugins/Editor/PluginBrowser/Source/PluginBrowser/Private/PluginMetadataObject.h
@@ -122,7 +122,7 @@ public:
 	bool bExplicitlyLoaded;
 
 	/** If true, prevents other plugins from depending on this plugin */
-	UPROPERTY(VisibleAnywhere, Category = Details)
+	UPROPERTY(EditAnywhere, Category = Details)
 	bool bIsSealed;
 
 	/** Plugins used by this plugin */
