Don't show toast about new plugins if the Plugin Editor tool isn't accessible
#rnx

diff --git a/Engine/Plugins/Editor/PluginBrowser/Source/PluginBrowser/Private/PluginBrowserModule.cpp b/Engine/Plugins/Editor/PluginBrowser/Source/PluginBrowser/Private/PluginBrowserModule.cpp
--- a/Engine/Plugins/Editor/PluginBrowser/Source/PluginBrowser/Private/PluginBrowserModule.cpp
+++ b/Engine/Plugins/Editor/PluginBrowser/Source/PluginBrowser/Private/PluginBrowserModule.cpp
@@ -13,6 +13,7 @@
 #include "Widgets/Notifications/SNotificationList.h"
 #include "Framework/Notifications/NotificationManager.h"
 #include "Misc/ConfigCacheIni.h"
+#include "Misc/NamePermissionList.h"
 #include "Interfaces/IPluginManager.h"
 #include "PluginDescriptorEditor.h"
 
@@ -194,7 +195,10 @@ TSharedRef<SDockTab> FPluginBrowserModule::SpawnPluginCreatorTab(const FSpawnTab
 void FPluginBrowserModule::OnMainFrameLoaded(TSharedPtr<SWindow> InRootWindow, bool bIsRunningStartupDialog)
 {
 	// Show a popup notification that allows the user to enable any new plugins
-	if(!bIsRunningStartupDialog && NewlyInstalledPlugins.Num() > 0 && !PluginBrowserTab.IsValid())
+	if(!bIsRunningStartupDialog 
+		&& NewlyInstalledPlugins.Num() > 0
+		&& !PluginBrowserTab.IsValid()
+		&& FGlobalTabmanager::Get()->GetTabPermissionList()->PassesFilter(PluginsEditorTabName))
 	{
 		FNotificationInfo Info(LOCTEXT("NewPluginsPopupTitle", "New plugins are available"));
 		Info.bFireAndForget = true;
