Attempt to replace the .uplugin file name with the desired template's plugin name.
The PLUGIN_NAME is replaced with the disired name.

diff --git a/Engine/Plugins/Editor/PluginBrowser/Source/PluginBrowser/Private/SNewPluginWizard.cpp b/Engine/Plugins/Editor/PluginBrowser/Source/PluginBrowser/Private/SNewPluginWizard.cpp
--- a/Engine/Plugins/Editor/PluginBrowser/Source/PluginBrowser/Private/SNewPluginWizard.cpp
+++ b/Engine/Plugins/Editor/PluginBrowser/Source/PluginBrowser/Private/SNewPluginWizard.cpp
@@ -586,6 +586,21 @@ FText SNewPluginWizard::GetPluginDestinationPath() const
 
 FText SNewPluginWizard::GetCurrentPluginName() const
 {
+	TSharedPtr<FPluginTemplateDescription> Template = PluginWizardDefinition->GetSelectedTemplate();
+	if (Template.IsValid())
+	{
+		TArray<FString> UPluginFiles;
+		IFileManager::Get().FindFiles(UPluginFiles, *Template->OnDiskPath, TEXT("*.uplugin"));
+		for (FString& File : UPluginFiles)
+		{
+			File.RemoveFromEnd(TEXT(".uplugin"));
+			if (File.Contains(TEXT("PLUGIN_NAME")))
+			{
+				return FText::FromString(File.Replace(TEXT("PLUGIN_NAME"), *PluginNameText.ToString(), ESearchCase::CaseSensitive));
+			}
+		}
+	}
+
 	return PluginNameText;
 }
 
@@ -615,7 +630,7 @@ void SNewPluginWizard::OnEnginePluginCheckboxChanged(ECheckBoxState NewCheckedSt
 
 FReply SNewPluginWizard::OnCreatePluginClicked()
 {
-	if (!ensure(!PluginFolderPath.IsEmpty() && !PluginNameText.IsEmpty()))
+	if (!ensure(!PluginFolderPath.IsEmpty() && !GetCurrentPluginName().IsEmpty()))
 	{
 		// Don't even try to assemble the path or else it may be relative to the binaries folder!
 		return FReply::Unhandled();
@@ -635,7 +650,7 @@ FReply SNewPluginWizard::OnCreatePluginClicked()
 		return FReply::Unhandled();
 	}
 
-	const FString PluginName = PluginNameText.ToString();
+	const FString PluginName = GetCurrentPluginName().ToString();
 	const bool bHasModules = PluginWizardDefinition->HasModules();
 	
 	FPluginUtils::FNewPluginParamsWithDescriptor CreationParams;
