Fix an issue where plugin names would compound when making a template causing name expansion.

diff --git a/Engine/Plugins/Editor/PluginBrowser/Source/PluginBrowser/Private/SNewPluginWizard.cpp b/Engine/Plugins/Editor/PluginBrowser/Source/PluginBrowser/Private/SNewPluginWizard.cpp
--- a/Engine/Plugins/Editor/PluginBrowser/Source/PluginBrowser/Private/SNewPluginWizard.cpp
+++ b/Engine/Plugins/Editor/PluginBrowser/Source/PluginBrowser/Private/SNewPluginWizard.cpp
@@ -586,21 +586,6 @@ FText SNewPluginWizard::GetPluginDestinationPath() const
 
 FText SNewPluginWizard::GetCurrentPluginName() const
 {
-	TSharedPtr<FPluginTemplateDescription> Template = PluginWizardDefinition->GetSelectedTemplate();
-	if (Template.IsValid() && !PluginNameText.IsEmpty())
-	{
-		TArray<FString> UPluginFiles;
-		IFileManager::Get().FindFiles(UPluginFiles, *Template->OnDiskPath, TEXT("*.uplugin"));
-		for (FString& File : UPluginFiles)
-		{
-			File.RemoveFromEnd(TEXT(".uplugin"));
-			if (File.Contains(TEXT("PLUGIN_NAME")))
-			{
-				return FText::FromString(File.Replace(TEXT("PLUGIN_NAME"), *PluginNameText.ToString(), ESearchCase::CaseSensitive));
-			}
-		}
-	}
-
 	return PluginNameText;
 }
 
@@ -650,7 +635,21 @@ FReply SNewPluginWizard::OnCreatePluginClicked()
 		return FReply::Unhandled();
 	}
 
-	const FString PluginName = GetCurrentPluginName().ToString();
+	FString PluginName;
+	TArray<FString> UPluginFiles;
+	IFileManager::Get().FindFiles(UPluginFiles, *Template->OnDiskPath, TEXT("*.uplugin"));
+	for (FString& File : UPluginFiles)
+	{
+		File.RemoveFromEnd(TEXT(".uplugin"));
+		if (File.Contains(TEXT("PLUGIN_NAME")))
+		{
+			PluginName = File.Replace(TEXT("PLUGIN_NAME"), *GetCurrentPluginName().ToString(), ESearchCase::CaseSensitive);
+		}
+	}
+	if (PluginName.IsEmpty())
+	{
+		PluginName = GetCurrentPluginName().ToString();
+	}
 	const bool bHasModules = PluginWizardDefinition->HasModules();
 	
 	FPluginUtils::FNewPluginParamsWithDescriptor CreationParams;
@@ -686,7 +685,7 @@ FReply SNewPluginWizard::OnCreatePluginClicked()
 	Template->CustomizeDescriptorBeforeCreation(CreationParams.Descriptor);
 
 
-	TSharedPtr<IPlugin> NewPlugin = FPluginUtils::CreateAndLoadNewPlugin(PluginName, PluginNameText.ToString(), PluginFolderPath, CreationParams, LoadParams);
+	TSharedPtr<IPlugin> NewPlugin = FPluginUtils::CreateAndLoadNewPlugin(PluginName, GetCurrentPluginName().ToString(), PluginFolderPath, CreationParams, LoadParams);
 	const bool bSucceeded = NewPlugin.IsValid();
 
 
