Plugin context menu actions in content browser path tree.
- Edit properties
- Package plugin
- Open reference viewer
- Support url
- Docs url
- Vendor url

diff --git a/Engine/Plugins/Editor/PluginBrowser/Source/PluginBrowser/PluginBrowser.Build.cs b/Engine/Plugins/Editor/PluginBrowser/Source/PluginBrowser/PluginBrowser.Build.cs
--- a/Engine/Plugins/Editor/PluginBrowser/Source/PluginBrowser/PluginBrowser.Build.cs
+++ b/Engine/Plugins/Editor/PluginBrowser/Source/PluginBrowser/PluginBrowser.Build.cs
@@ -38,7 +38,10 @@ namespace UnrealBuildTool.Rules
                     "AssetTools",
 					"Json",
 					"ToolWidgets",
+					"ToolMenus",
 					"EditorWidgets",
+					"ContentBrowser",
+					"ContentBrowserData",
 				}
 			);
 
diff --git a/Engine/Plugins/Editor/PluginBrowser/Source/PluginBrowser/Private/PluginActions.cpp b/Engine/Plugins/Editor/PluginBrowser/Source/PluginBrowser/Private/PluginActions.cpp
new file mode 100644
--- /dev/null
+++ b/Engine/Plugins/Editor/PluginBrowser/Source/PluginBrowser/Private/PluginActions.cpp
@@ -0,0 +1,43 @@
+// Copyright Epic Games, Inc. All Rights Reser
+
+#include "PluginActions.h"
+
+#include "DesktopPlatformModule.h"
+#include "Framework/Application/SlateApplication.h"
+#include "IDesktopPlatform.h"
+#include "Interfaces/IPluginManager.h"
+#include "IUATHelperModule.h"
+
+#define LOCTEXT_NAMESPACE "PluginListTile"
+
+void PluginActions::PackagePlugin(TSharedRef<IPlugin> Plugin, TSharedPtr<SWidget> ParentWidget)
+{
+	FString DefaultDirectory;
+	FString OutputDirectory;
+
+	if ( !FDesktopPlatformModule::Get()->OpenDirectoryDialog(FSlateApplication::Get().FindBestParentWindowHandleForDialogs(ParentWidget), LOCTEXT("PackagePluginDialogTitle", "Package Plugin...").ToString(), DefaultDirectory, OutputDirectory) )
+	{
+		return;
+	}
+
+	// Ensure path is full rather than relative (for macs)
+	FString DescriptorFilename = Plugin->GetDescriptorFileName();
+	FString DescriptorFullPath = FPaths::ConvertRelativePathToFull(DescriptorFilename);
+	OutputDirectory = FPaths::Combine(OutputDirectory, Plugin->GetName());
+	FString CommandLine = FString::Printf(TEXT("BuildPlugin -Plugin=\"%s\" -Package=\"%s\" -CreateSubFolder"), *DescriptorFullPath, *OutputDirectory);
+
+#if PLATFORM_WINDOWS
+	FText PlatformName = LOCTEXT("PlatformName_Windows", "Windows");
+#elif PLATFORM_MAC
+	FText PlatformName = LOCTEXT("PlatformName_Mac", "Mac");
+#elif PLATFORM_LINUX
+	FText PlatformName = LOCTEXT("PlatformName_Linux", "Linux");
+#else
+	FText PlatformName = LOCTEXT("PlatformName_Other", "Other OS");
+#endif
+
+	IUATHelperModule::Get().CreateUatTask(CommandLine, PlatformName, LOCTEXT("PackagePluginTaskName", "Packaging Plugin"),
+		LOCTEXT("PackagePluginTaskShortName", "Package Plugin Task"), FAppStyle::GetBrush(TEXT("MainFrame.CookContent")));
+}
+
+#undef LOCTEXT_NAMESPACE
\ No newline at end of file
diff --git a/Engine/Plugins/Editor/PluginBrowser/Source/PluginBrowser/Private/PluginActions.h b/Engine/Plugins/Editor/PluginBrowser/Source/PluginBrowser/Private/PluginActions.h
new file mode 100644
--- /dev/null
+++ b/Engine/Plugins/Editor/PluginBrowser/Source/PluginBrowser/Private/PluginActions.h
@@ -0,0 +1,13 @@
+// Copyright Epic Games, Inc. All Rights Reserved.
+
+#pragma once
+
+#include "Templates/SharedPointer.h"
+
+class IPlugin; 
+class SWidget; 
+
+namespace PluginActions
+{
+    void PackagePlugin(TSharedRef<IPlugin> Plugin, TSharedPtr<SWidget> ParentWidget);
+}
\ No newline at end of file
diff --git a/Engine/Plugins/Editor/PluginBrowser/Source/PluginBrowser/Private/PluginBrowserModule.cpp b/Engine/Plugins/Editor/PluginBrowser/Source/PluginBrowser/Private/PluginBrowserModule.cpp
--- a/Engine/Plugins/Editor/PluginBrowser/Source/PluginBrowser/Private/PluginBrowserModule.cpp
+++ b/Engine/Plugins/Editor/PluginBrowser/Source/PluginBrowser/Private/PluginBrowserModule.cpp
@@ -1,21 +1,29 @@
 // Copyright Epic Games, Inc. All Rights Reserved.
 
 #include "PluginBrowserModule.h"
-#include "SPluginBrowser.h"
+
+#include "Algo/AnyOf.h"
+#include "ContentBrowserModule.h"
 #include "Features/IModularFeatures.h"
 #include "Features/EditorFeatures.h"
-#include "PluginMetadataObject.h"
-#include "PluginStyle.h"
-#include "PropertyEditorModule.h"
-#include "Widgets/Docking/SDockTab.h"
-#include "SNewPluginWizard.h"
-#include "Interfaces/IMainFrameModule.h"
-#include "Widgets/Notifications/SNotificationList.h"
+#include "Framework/MultiBox/MultiBoxBuilder.h"
 #include "Framework/Notifications/NotificationManager.h"
+#include "Interfaces/IPluginManager.h"
+#include "Interfaces/IMainFrameModule.h"
 #include "Misc/ConfigCacheIni.h"
 #include "Misc/NamePermissionList.h"
-#include "Interfaces/IPluginManager.h"
+#include "Misc/PathViews.h"
+#include "PropertyEditorModule.h"
+#include "PluginActions.h"
 #include "PluginDescriptorEditor.h"
+#include "PluginMetadataObject.h"
+#include "PluginStyle.h"
+#include "SNewPluginWizard.h"
+#include "SPluginBrowser.h"
+#include "ToolMenu.h"
+#include "ToolMenus.h"
+#include "Widgets/Docking/SDockTab.h"
+#include "Widgets/Notifications/SNotificationList.h"
 
 #define LOCTEXT_NAMESPACE "PluginsEditor"
 
@@ -81,6 +89,8 @@ void FPluginBrowserModule::StartupModule()
 	// Register a callback to check for new plugins on startup
 	IMainFrameModule& MainFrameModule = FModuleManager::LoadModuleChecked<IMainFrameModule>(TEXT("MainFrame"));
 	MainFrameModule.OnMainFrameCreationFinished().AddRaw(this, &FPluginBrowserModule::OnMainFrameLoaded);
+	
+	AddContentBrowserMenuExtensions();
 }
 
 void FPluginBrowserModule::ShutdownModule()
@@ -102,6 +112,125 @@ void FPluginBrowserModule::ShutdownModule()
 	IModularFeatures::Get().UnregisterModularFeature( EditorFeatures::PluginsEditor, this );
 }
 
+static void CreatePluginMenu(UToolMenu* Menu, TSharedRef<IPlugin> Plugin)
+{
+	UContentBrowserDataMenuContext_FolderMenu* Context = Menu->FindContext<UContentBrowserDataMenuContext_FolderMenu>();
+	
+	FToolMenuSection& Section = Menu->FindOrAddSection("PluginMenu");
+	Section.AddMenuEntry(
+		"EditPluginProperties",
+		LOCTEXT("PluginContextMenu.EditPlugin", "Edit"),
+		LOCTEXT("PluginContextMenu.EditPlugin.ToolTip", "Edit this plugin's properties."),
+		FSlateIcon(FAppStyle::GetAppStyleSetName(), "Icons.Edit"),
+		FExecuteAction::CreateLambda([Plugin](){
+			FPluginBrowserModule::Get().OpenPluginEditor(Plugin, {}, {});
+		})	
+	);
+
+	Section.AddMenuEntry(
+		"PackagePlugin",
+		LOCTEXT("PluginContextMenu.PackagePlugin", "Package"),
+		LOCTEXT("PluginContextMenu.PackagePlugin.ToolTip", "Package this plugin for distribution"),
+		FSlateIcon(FAppStyle::GetAppStyleSetName(), "Icons.Package"),
+		FExecuteAction::CreateLambda([Plugin](){
+			PluginActions::PackagePlugin(Plugin, {});
+		})	
+	);
+
+	if (FPluginBrowserModule::Get().OnLaunchReferenceViewerDelegate().IsBound())
+	{
+		Section.AddMenuEntry(
+			"PluginReferenceViewer",
+			LOCTEXT("PluginContextMenu.ReferenceViewer", "Reference Viewer"),
+			LOCTEXT("PluginContextMenu.ReferenceViewer.ToolTip", "Open the plugin reference viewer showing references to and from this plugin"),
+			FSlateIcon(FAppStyle::GetAppStyleSetName(), "ContentBrowser.ReferenceViewer"),
+			FExecuteAction::CreateLambda([Plugin](){
+				FPluginBrowserModule::Get().OnLaunchReferenceViewerDelegate().Execute(Plugin);
+			})	
+		);
+	}
+
+	const FPluginDescriptor& Descriptor = Plugin->GetDescriptor();
+	if (Descriptor.SupportURL.StartsWith(TEXT("http://")))
+	{
+		Section.AddMenuEntry(
+			"PluginSupport",
+			LOCTEXT("PluginContextMenu.Support", "Support"),
+			LOCTEXT("PluginContextMenu.Support.ToolTip", "Open the plugin's online support URL"),
+			FSlateIcon(FAppStyle::GetAppStyleSetName(), "Icons.Comment"),
+			FExecuteAction::CreateLambda([Plugin](){
+				const FPluginDescriptor& Descriptor = Plugin->GetDescriptor();
+				if (Descriptor.SupportURL.StartsWith(TEXT("http://")))
+				{
+					FPlatformProcess::LaunchURL(*Descriptor.SupportURL, nullptr, nullptr);
+				}
+			})	
+		);
+	}
+	
+	if (Descriptor.DocsURL.StartsWith(TEXT("http://")))
+	{	
+		Section.AddMenuEntry(
+			"PluginDocumentation",
+			LOCTEXT("PluginContextMenu.Documentation", "Documentation"),
+			LOCTEXT("PluginContextMenu.Documentation.ToolTip", "Open the plugin's online support URL."),
+			FSlateIcon(FAppStyle::GetAppStyleSetName(), "Icons.Documentation"),
+			FExecuteAction::CreateLambda([Plugin](){
+				const FPluginDescriptor& Descriptor = Plugin->GetDescriptor();
+				if (Descriptor.DocsURL.StartsWith(TEXT("http://")))
+				{
+					FPlatformProcess::LaunchURL(*Descriptor.DocsURL, nullptr, nullptr);
+				}
+			})	
+		);
+	}
+	
+	if (Descriptor.CreatedBy.IsEmpty() == false && Descriptor.CreatedByURL.StartsWith(TEXT("http://")))
+	{
+		Section.AddMenuEntry(
+			"PluginVendor",
+			LOCTEXT("PluginContextMenu.Vendor", "Vendor"),
+			LOCTEXT("PluginContextMenu.Vendor.ToolTip", "Visit the plugin vendor's website."),
+			FSlateIcon(FAppStyle::GetAppStyleSetName(), "Icons.OpenInBrowser"),
+			FExecuteAction::CreateLambda([Plugin](){
+				const FPluginDescriptor& Descriptor = Plugin->GetDescriptor();
+				if (!Descriptor.CreatedBy.IsEmpty() && Descriptor.CreatedByURL.StartsWith(TEXT("http://")))
+				{
+					FPlatformProcess::LaunchURL(*Descriptor.CreatedByURL, nullptr, nullptr);
+				}
+			})	
+		);
+	}
+}
+
+void FPluginBrowserModule::AddContentBrowserMenuExtensions()
+{
+	if (UToolMenu* Menu = UToolMenus::Get()->ExtendMenu("ContentBrowser.FolderContextMenu"))
+	{
+		FToolMenuSection& Section = Menu->AddDynamicSection(TEXT("DynamicSection_Plugins"), FNewToolMenuDelegate::CreateLambda([](UToolMenu* InMenu)
+		{
+			UContentBrowserDataMenuContext_FolderMenu* Context = InMenu->FindContext<UContentBrowserDataMenuContext_FolderMenu>();
+			// Only show the menu if there is a single item selected as most commands don't work nicely for multiple plugins 
+			if (Context->SelectedItems.Num() == 1 && Context->SelectedItems[0].IsInPlugin())
+			{
+				FNameBuilder ItemPath{Context->SelectedItems[0].GetInternalPath()};
+				FStringView PluginName = FPathViews::GetMountPointNameFromPath(ItemPath.ToView());
+				if (TSharedPtr<IPlugin> Plugin = IPluginManager::Get().FindPlugin(PluginName))
+				{
+					FToolMenuSection& Section = InMenu->FindOrAddSection("Section");
+					Section.AddSubMenu("Plugin", 
+						LOCTEXT("ContentBrowserPluginMenuName", "Plugin"),
+						LOCTEXT("ContentBrowserPluginMenuTooltip", "Plugin Actions"),
+						FNewToolMenuDelegate::CreateStatic(&CreatePluginMenu, Plugin.ToSharedRef()),
+						false,
+						FSlateIcon(FPluginStyle::Get()->GetStyleSetName(), "Plugins.TabIcon")
+						);
+				}
+			}
+		}));
+	}
+}
+
 void FPluginBrowserModule::RegisterPluginTemplate(TSharedRef<FPluginTemplateDescription> Template)
 {
 	LLM_SCOPE_BYTAG(PluginBrowser);
diff --git a/Engine/Plugins/Editor/PluginBrowser/Source/PluginBrowser/Private/PluginBrowserModule.h b/Engine/Plugins/Editor/PluginBrowser/Source/PluginBrowser/Private/PluginBrowserModule.h
--- a/Engine/Plugins/Editor/PluginBrowser/Source/PluginBrowser/Private/PluginBrowserModule.h
+++ b/Engine/Plugins/Editor/PluginBrowser/Source/PluginBrowser/Private/PluginBrowserModule.h
@@ -101,6 +101,9 @@ private:
 	/** Updates the user's config file with the list of installed plugins that they've seen. */
 	void UpdatePreviousInstalledPlugins();
 
+	/** Register menu extensions for the content browser */
+	void AddContentBrowserMenuExtensions();
+
 	/** List of added plugin templates */
 	TArray<TSharedRef<FPluginTemplateDescription>> AddedPluginTemplates;
 
diff --git a/Engine/Plugins/Editor/PluginBrowser/Source/PluginBrowser/Private/SPluginTile.cpp b/Engine/Plugins/Editor/PluginBrowser/Source/PluginBrowser/Private/SPluginTile.cpp
--- a/Engine/Plugins/Editor/PluginBrowser/Source/PluginBrowser/Private/SPluginTile.cpp
+++ b/Engine/Plugins/Editor/PluginBrowser/Source/PluginBrowser/Private/SPluginTile.cpp
@@ -4,7 +4,6 @@
 #include "Brushes/SlateDynamicImageBrush.h"
 #include "HAL/PlatformFileManager.h"
 #include "Misc/MessageDialog.h"
-#include "IDesktopPlatform.h"
 #include "Misc/App.h"
 #include "Framework/Application/SlateApplication.h"
 #include "PluginReferenceDescriptor.h"
@@ -13,13 +12,13 @@
 #include "Widgets/Input/SCheckBox.h"
 #include "Interfaces/IPluginManager.h"
 #include "SPluginBrowser.h"
+#include "PluginActions.h"
 #include "PluginStyle.h"
 #include "GameProjectGenerationModule.h"
 #include "SPluginTileList.h"
 #include "Widgets/Input/SHyperlink.h"
 #include "Interfaces/IProjectManager.h"
 #include "PluginBrowserModule.h"
-#include "IUATHelperModule.h"
 #include "DesktopPlatformModule.h"
 #include "Misc/ConfigCacheIni.h"
 #include "Widgets/Layout/SBox.h"
@@ -692,32 +691,10 @@ void SPluginTile::OnEditPluginFinished()
 
 void SPluginTile::OnPackagePlugin()
 {
-	FString DefaultDirectory;
-	FString OutputDirectory;
-
-	if ( !FDesktopPlatformModule::Get()->OpenDirectoryDialog(FSlateApplication::Get().FindBestParentWindowHandleForDialogs(AsShared()), LOCTEXT("PackagePluginDialogTitle", "Package Plugin...").ToString(), DefaultDirectory, OutputDirectory) )
+	if (ensure(Plugin.IsValid()))
 	{
-		return;
+		PluginActions::PackagePlugin(Plugin.ToSharedRef(), AsShared());
 	}
-
-	// Ensure path is full rather than relative (for macs)
-	FString DescriptorFilename = Plugin->GetDescriptorFileName();
-	FString DescriptorFullPath = FPaths::ConvertRelativePathToFull(DescriptorFilename);
-	OutputDirectory = FPaths::Combine(OutputDirectory, Plugin->GetName());
-	FString CommandLine = FString::Printf(TEXT("BuildPlugin -Plugin=\"%s\" -Package=\"%s\" -CreateSubFolder"), *DescriptorFullPath, *OutputDirectory);
-
-#if PLATFORM_WINDOWS
-	FText PlatformName = LOCTEXT("PlatformName_Windows", "Windows");
-#elif PLATFORM_MAC
-	FText PlatformName = LOCTEXT("PlatformName_Mac", "Mac");
-#elif PLATFORM_LINUX
-	FText PlatformName = LOCTEXT("PlatformName_Linux", "Linux");
-#else
-	FText PlatformName = LOCTEXT("PlatformName_Other", "Other OS");
-#endif
-
-	IUATHelperModule::Get().CreateUatTask(CommandLine, PlatformName, LOCTEXT("PackagePluginTaskName", "Packaging Plugin"),
-		LOCTEXT("PackagePluginTaskShortName", "Package Plugin Task"), FAppStyle::GetBrush(TEXT("MainFrame.CookContent")));
 }
 
 void SPluginTile::OnOpenPluginReferenceViewer()
