Fix copyright boilerplate

diff --git a/Engine/Plugins/Editor/PluginBrowser/Source/PluginBrowser/Private/PluginActions.cpp b/Engine/Plugins/Editor/PluginBrowser/Source/PluginBrowser/Private/PluginActions.cpp
--- a/Engine/Plugins/Editor/PluginBrowser/Source/PluginBrowser/Private/PluginActions.cpp
+++ b/Engine/Plugins/Editor/PluginBrowser/Source/PluginBrowser/Private/PluginActions.cpp
@@ -1,4 +1,4 @@
-// Copyright Epic Games, Inc. All Rights Reser
+// Copyright Epic Games, Inc. All Rights Reserved.
 
 #include "PluginActions.h"
 
