Add Plugin submenu to asset context menu in content browser

diff --git a/Engine/Plugins/Editor/PluginBrowser/Source/PluginBrowser/Private/PluginBrowserModule.cpp b/Engine/Plugins/Editor/PluginBrowser/Source/PluginBrowser/Private/PluginBrowserModule.cpp
--- a/Engine/Plugins/Editor/PluginBrowser/Source/PluginBrowser/Private/PluginBrowserModule.cpp
+++ b/Engine/Plugins/Editor/PluginBrowser/Source/PluginBrowser/Private/PluginBrowserModule.cpp
@@ -203,6 +203,17 @@ static void CreatePluginMenu(UToolMenu* Menu, TSharedRef<IPlugin> Plugin)
 	}
 }
 
+void AddPluginMenuToContextMenu(UToolMenu* InMenu, const TSharedPtr<IPlugin>& Plugin, FName AddToSection)
+{
+	FToolMenuSection &Section = InMenu->FindOrAddSection(AddToSection);
+	Section.AddSubMenu("Plugin",
+	   LOCTEXT("ContentBrowserPluginMenuName", "Plugin"),
+	   LOCTEXT("ContentBrowserPluginMenuTooltip", "Plugin Actions"),
+	   FNewToolMenuDelegate::CreateStatic(&CreatePluginMenu, Plugin.ToSharedRef()),
+	   false,
+	   FSlateIcon(FPluginStyle::Get()->GetStyleSetName(), "Plugins.TabIcon"));
+}
+
 void FPluginBrowserModule::AddContentBrowserMenuExtensions()
 {
 	if (UToolMenu* Menu = UToolMenus::Get()->ExtendMenu("ContentBrowser.FolderContextMenu"))
@@ -217,14 +228,25 @@ void FPluginBrowserModule::AddContentBrowserMenuExtensions()
 				FStringView PluginName = FPathViews::GetMountPointNameFromPath(ItemPath.ToView());
 				if (TSharedPtr<IPlugin> Plugin = IPluginManager::Get().FindPlugin(PluginName))
 				{
-					FToolMenuSection& Section = InMenu->FindOrAddSection("Section");
-					Section.AddSubMenu("Plugin", 
-						LOCTEXT("ContentBrowserPluginMenuName", "Plugin"),
-						LOCTEXT("ContentBrowserPluginMenuTooltip", "Plugin Actions"),
-						FNewToolMenuDelegate::CreateStatic(&CreatePluginMenu, Plugin.ToSharedRef()),
-						false,
-						FSlateIcon(FPluginStyle::Get()->GetStyleSetName(), "Plugins.TabIcon")
-						);
+					AddPluginMenuToContextMenu(InMenu, Plugin, "Section");
+				}
+			}
+		}));
+	}
+	
+	if (UToolMenu* Menu = UToolMenus::Get()->ExtendMenu("ContentBrowser.AssetContextMenu"))
+	{
+		FToolMenuSection& Section = Menu->AddDynamicSection(TEXT("DynamicSection_Plugins"), FNewToolMenuDelegate::CreateLambda([](UToolMenu* InMenu)
+		{
+			UContentBrowserDataMenuContext_FileMenu* Context = InMenu->FindContext<UContentBrowserDataMenuContext_FileMenu>();
+			// Only show the menu if there is a single item selected as most commands don't work nicely for multiple plugins 
+			if (Context->SelectedItems.Num() == 1 && Context->SelectedItems[0].IsInPlugin())
+			{
+				FNameBuilder ItemPath{Context->SelectedItems[0].GetInternalPath()};
+				FStringView PluginName = FPathViews::GetMountPointNameFromPath(ItemPath.ToView());
+				if (TSharedPtr<IPlugin> Plugin = IPluginManager::Get().FindPlugin(PluginName))
+				{
+					AddPluginMenuToContextMenu(InMenu, Plugin, "AssetContextExploreMenuOptions");
 				}
 			}
 		}));
