Plugin Browser + Plugin Manager: Add UI for configuring additional/external plugin directories.

Adds a "Plugin Directories" button to the Plugin Browser, which spawns a nomad tab containing a struct details view that lists configured directories, and optionally (config-driven) allows them to be modified and handles updates to the relevant managers.

Also includes some smaller supporting changes to make the Plugin Browser usable in LiveLink Hub.

Related changes to plugin manager:
- Adds an EPluginExternalSource enum, and extends IPluginManager's "external path" representation from FString to a tagged struct indicating the source of each path.
- Adds `void IPluginManager::GetExternalPluginSources(TSet<FExternalPluginPath>& OutPluginSources)` as an alternative/successor to `const TSet<FString>& GetAdditionalPluginSearchPaths()`.
- Adds `IPluginManager::RemovePluginSearchPath()` counterpart to existing `AddPluginSearchPath()`.

Other changes:
- Adds `[EditorSettings] bCanModifyProjectPluginDirectoriesFromBrowser` config bool. If not explicitly set, effectively defaults to the value of `bCanModifyPluginsFromBrowser`.
- Adds `[EditorSettings] bCanModifyUserPluginDirectoriesFromBrowser` config bool. Disabled and hidden by default, 'User' plugin directories are intended as an application-specific mechanism to expose `IPluginManager::AddPluginSearchPath` without requiring a writable uproject descriptor.
- Adds `[EditorSettings] HidePluginCategoriesFromBrowser` config string list to exclude named categories from SPluginCategoryTree (unused by default, will be used for LiveLink Hub).
- Adds `[EditorSettings] bCanCreatePluginsFromBrowser` config bool to disable the plugin creator / "Add Plugin" button.

diff --git a/Engine/Plugins/Editor/PluginBrowser/Source/PluginBrowser/Private/PluginBrowserModule.cpp b/Engine/Plugins/Editor/PluginBrowser/Source/PluginBrowser/Private/PluginBrowserModule.cpp
--- a/Engine/Plugins/Editor/PluginBrowser/Source/PluginBrowser/Private/PluginBrowserModule.cpp
+++ b/Engine/Plugins/Editor/PluginBrowser/Source/PluginBrowser/Private/PluginBrowserModule.cpp
@@ -20,6 +20,7 @@
 #include "PluginStyle.h"
 #include "SNewPluginWizard.h"
 #include "SPluginBrowser.h"
+#include "SPluginPaths.h"
 #include "ToolMenu.h"
 #include "ToolMenus.h"
 #include "Widgets/Docking/SDockTab.h"
@@ -29,10 +30,13 @@
 
 LLM_DEFINE_TAG(PluginBrowser);
 
+DEFINE_LOG_CATEGORY(LogPluginBrowser);
+
 IMPLEMENT_MODULE( FPluginBrowserModule, PluginBrowser )
 
 const FName FPluginBrowserModule::PluginsEditorTabName( TEXT( "PluginsEditor" ) );
 const FName FPluginBrowserModule::PluginCreatorTabName( TEXT( "PluginCreator" ) );
+const FName FPluginBrowserModule::ExternalDirectoriesTabName(TEXT("PluginDirectories"));
 
 void FPluginBrowserModule::StartupModule()
 {
@@ -60,8 +64,17 @@ void FPluginBrowserModule::StartupModule()
 		.SetMenuType(ETabSpawnerMenuType::Hidden);
 
 	// Register a default size for this tab
-	FVector2D DefaultSize(1000.0f, 750.0f);
-	FTabManager::RegisterDefaultTabWindowSize(PluginCreatorTabName, DefaultSize);
+	const FVector2D DefaultCreatorSize(1000.0f, 750.0f);
+	FTabManager::RegisterDefaultTabWindowSize(PluginCreatorTabName, DefaultCreatorSize);
+
+	// Additional Plugin Directories tab
+	FGlobalTabmanager::Get()->RegisterNomadTabSpawner(ExternalDirectoriesTabName, FOnSpawnTab::CreateRaw(this, &FPluginBrowserModule::HandleSpawnExternalDirectoriesTab))
+		.SetDisplayName(LOCTEXT("ExternalPluginDirectoriesTabTitle", "Plugin Directories"))
+		.SetMenuType(ETabSpawnerMenuType::Hidden)
+		.SetIcon(FSlateIcon(FPluginStyle::Get()->GetStyleSetName(), "Plugins.TabIcon"));
+
+	const FVector2D DefaultDirectoriesSize(750.0f, 600.0f);
+	FTabManager::RegisterDefaultTabWindowSize(ExternalDirectoriesTabName, DefaultDirectoriesSize);
 
 	// Get a list of the installed plugins we've seen before
 	TArray<FString> PreviousInstalledPlugins;
@@ -91,6 +104,8 @@ void FPluginBrowserModule::StartupModule()
 	MainFrameModule.OnMainFrameCreationFinished().AddRaw(this, &FPluginBrowserModule::OnMainFrameLoaded);
 	
 	AddContentBrowserMenuExtensions();
+
+	IPluginManager::Get().GetExternalPluginSources(OriginalExternalSources);
 }
 
 void FPluginBrowserModule::ShutdownModule()
@@ -107,6 +122,7 @@ void FPluginBrowserModule::ShutdownModule()
 	// Unregister the tab spawner
 	FGlobalTabmanager::Get()->UnregisterTabSpawner( PluginsEditorTabName );
 	FGlobalTabmanager::Get()->UnregisterNomadTabSpawner( PluginCreatorTabName );
+	FGlobalTabmanager::Get()->UnregisterNomadTabSpawner(ExternalDirectoriesTabName);
 
 	// Unregister our feature
 	IModularFeatures::Get().UnregisterModularFeature( EditorFeatures::PluginsEditor, this );
@@ -312,6 +328,19 @@ bool FPluginBrowserModule::IsNewlyInstalledPlugin(const FString& PluginName) con
 	return NewlyInstalledPlugins.Contains(PluginName);
 }
 
+bool FPluginBrowserModule::ShowPendingRestart() const
+{
+	LastQueriedExternalSources.Reset();
+	IPluginManager::Get().GetExternalPluginSources(LastQueriedExternalSources);
+
+	// For sets of equal size, inclusion denotes equality.
+	const bool bExternalSourcesChanged =
+		(LastQueriedExternalSources.Num() != OriginalExternalSources.Num())
+		|| !OriginalExternalSources.Includes(LastQueriedExternalSources);
+
+	return HasPluginsPendingEnable() || bExternalSourcesChanged;
+}
+
 TSharedRef<SDockTab> FPluginBrowserModule::HandleSpawnPluginBrowserTab(const FSpawnTabArgs& SpawnTabArgs)
 {
 	const TSharedRef<SDockTab> MajorTab = 
@@ -343,6 +372,17 @@ TSharedRef<SDockTab> FPluginBrowserModule::SpawnPluginCreatorTab(const FSpawnTab
 	return ResultTab;
 }
 
+TSharedRef<SDockTab> FPluginBrowserModule::HandleSpawnExternalDirectoriesTab(const FSpawnTabArgs& SpawnTabArgs)
+{
+	const TSharedRef<SDockTab> MajorTab = SNew(SDockTab)
+		.TabRole(ETabRole::NomadTab);
+
+	TSharedRef<SWidget> TabContent = SNew(SPluginPaths);
+	MajorTab->SetContent(TabContent);
+
+	return MajorTab;
+}
+
 void FPluginBrowserModule::OnMainFrameLoaded(TSharedPtr<SWindow> InRootWindow, bool bIsRunningStartupDialog)
 {
 	// Show a popup notification that allows the user to enable any new plugins
diff --git a/Engine/Plugins/Editor/PluginBrowser/Source/PluginBrowser/Private/PluginBrowserModule.h b/Engine/Plugins/Editor/PluginBrowser/Source/PluginBrowser/Private/PluginBrowserModule.h
--- a/Engine/Plugins/Editor/PluginBrowser/Source/PluginBrowser/Private/PluginBrowserModule.h
+++ b/Engine/Plugins/Editor/PluginBrowser/Source/PluginBrowser/Private/PluginBrowserModule.h
@@ -3,6 +3,7 @@
 #pragma once 
 
 #include "IPluginBrowser.h"
+#include "Interfaces/IPluginManager.h"
 
 class IModuleInterface;
 class IPluginWizardDefinition;
@@ -14,6 +15,8 @@ class FSpawnTabArgs;
 
 DECLARE_MULTICAST_DELEGATE(FOnNewPluginCreated);
 
+DECLARE_LOG_CATEGORY_EXTERN(LogPluginBrowser, Log, All);
+
 class FPluginBrowserModule : public IPluginBrowser
 {
 public:
@@ -41,6 +44,7 @@ public:
 	void BroadcastNewPluginCreated() const {NewPluginCreatedDelegate.Broadcast();}
 
 	virtual FOnLaunchReferenceViewer& OnLaunchReferenceViewerDelegate() override { return LaunchReferenceViewerDelegate; }
+	virtual FOnPluginDirectoriesChanged& OnPluginDirectoriesChanged() override { return OnPluginDirectoriesChangedDelegate; }
 
 	/**
 	 * Sets whether a plugin is pending enable/disable
@@ -69,12 +73,18 @@ public:
 	/** Checks whether the given plugin should be displayed with a 'NEW' label */
 	bool IsNewlyInstalledPlugin(const FString& PluginName) const;
 
+	/** Whether the restart editor notice should be displayed. */
+	bool ShowPendingRestart() const;
+
 	/** ID name for the plugins editor major tab */
 	static const FName PluginsEditorTabName;
 
 	/** ID name for the plugin creator tab */
 	static const FName PluginCreatorTabName;
 
+	/** ID name for the external plugin directories tab */
+	static const FName ExternalDirectoriesTabName;
+
 	/** Spawns the plugin creator tab with a specific wizard definition */
 	virtual TSharedRef<SDockTab> SpawnPluginCreatorTab(const FSpawnTabArgs& SpawnTabArgs, TSharedPtr<IPluginWizardDefinition> PluginWizardDefinition) override;
 
@@ -89,6 +99,9 @@ private:
 	/** Called to spawn the plugin creator tab */
 	TSharedRef<SDockTab> HandleSpawnPluginCreatorTab(const FSpawnTabArgs& SpawnTabArgs);
 
+	/** Called to spawn the external directories tab */
+	TSharedRef<SDockTab> HandleSpawnExternalDirectoriesTab(const FSpawnTabArgs& SpawnTabArgs);
+
 	/** Callback for the main frame finishing load */
 	void OnMainFrameLoaded(TSharedPtr<SWindow> InRootWindow, bool bIsRunningStartupDialog);
 
@@ -123,12 +136,21 @@ private:
 	/** List of plugins that have been recently installed */
 	TSet<FString> NewlyInstalledPlugins;
 
+	/** External plugin sources configuration as captured at startup. */
+	TSet<FExternalPluginPath> OriginalExternalSources;
+
+	/** Most recently queried plugin sources configuration. */
+	mutable TSet<FExternalPluginPath> LastQueriedExternalSources;
+
 	/** Delegate called when a new plugin is created */
 	FOnNewPluginCreated NewPluginCreatedDelegate;
 
 	/** Delegate that if bound the Plugin Browser will show the Reference Viewer button. Delegate called when the button is clicked */
 	FOnLaunchReferenceViewer LaunchReferenceViewerDelegate;
 
+	/** Called when the external plugin directories configuration is modified via the browser. */
+	FOnPluginDirectoriesChanged OnPluginDirectoriesChangedDelegate;
+
 	/** Notification popup that new plugins are available */
 	TWeakPtr<SNotificationItem> NewPluginsNotification;
 };
diff --git a/Engine/Plugins/Editor/PluginBrowser/Source/PluginBrowser/Private/SPluginBrowser.cpp b/Engine/Plugins/Editor/PluginBrowser/Source/PluginBrowser/Private/SPluginBrowser.cpp
--- a/Engine/Plugins/Editor/PluginBrowser/Source/PluginBrowser/Private/SPluginBrowser.cpp
+++ b/Engine/Plugins/Editor/PluginBrowser/Source/PluginBrowser/Private/SPluginBrowser.cpp
@@ -23,6 +23,7 @@
 #include "Styling/StyleColors.h"
 #include "SWarningOrErrorBox.h"
 #include "Widgets/Views/SListView.h" // IWYU pragma: keep
+#include "Misc/ConfigCacheIni.h"
 
 #define LOCTEXT_NAMESPACE "PluginsEditor"
 
@@ -145,19 +146,25 @@ void SPluginBrowser::Construct( const FArguments& Args )
 		+ SVerticalBox::Slot()
 		.AutoHeight()
 		[
-			// Add Plugin button
 			SNew(SHorizontalBox)
 
+			// Add Plugin button
 			+SHorizontalBox::Slot()
 			.VAlign(VAlign_Center)
 			.HAlign(HAlign_Left)
-			.Padding(FMargin(12, 7, 18, 7))
+			.Padding(FMargin(12, 7, 0, 7))
 			.AutoWidth()
 			[
 				SNew(SButton)
 				.ToolTip(SNew(SToolTip).Text(LOCTEXT("NewPluginEnabled", "Click here to open the Plugin Creator dialog.")))
-				.OnClicked(this, &SPluginBrowser::HandleNewPluginButtonClicked)
 				.ContentPadding(FMargin(0, 5.f, 0, 4.f))
+				.OnClicked(this, &SPluginBrowser::HandleNewPluginButtonClicked)
+				.IsEnabled_Static([]()
+				{
+					bool bCreateEnabled = true;
+					GConfig->GetBool(TEXT("EditorSettings"), TEXT("bCanCreatePluginsFromBrowser"), bCreateEnabled, GEditorIni);
+					return bCreateEnabled;
+				})
 				.Content()
 				[
 					SNew(SHorizontalBox)
@@ -181,6 +188,43 @@ void SPluginBrowser::Construct( const FArguments& Args )
 				]
 			]
 
+			// Plugin Directories button
+			+SHorizontalBox::Slot()
+			.VAlign(VAlign_Center)
+			.HAlign(HAlign_Left)
+			.Padding(FMargin(7, 7, 18, 7))
+			.AutoWidth()
+			[
+				SNew(SButton)
+				.ToolTip(SNew(SToolTip).Text(LOCTEXT("PluginDirectoriesTooltip", "Click here to configure additional plugin search directories.")))
+				.OnClicked_Static([]() -> FReply
+				{
+					FGlobalTabmanager::Get()->TryInvokeTab(FPluginBrowserModule::ExternalDirectoriesTabName);
+					return FReply::Handled();
+				})
+				.ContentPadding(FMargin(0, 5.f, 0, 4.f))
+				.Content()
+				[
+					SNew(SHorizontalBox)
+					+ SHorizontalBox::Slot()
+					.HAlign(HAlign_Center)
+					.VAlign(VAlign_Center)
+					[
+						SNew(SImage)
+						.Image(FPluginStyle::Get()->GetBrush("Plugins.TabIcon"))
+					]
+					+ SHorizontalBox::Slot()
+					.Padding(FMargin(3, 0, 0, 0))
+					.VAlign(VAlign_Center)
+					.AutoWidth()
+					[
+						SNew(STextBlock)
+						.TextStyle(FAppStyle::Get(), "SmallButtonText")
+						.Text(LOCTEXT("PluginDirectoriesLabel", "Plugin Directories"))
+					]
+				]
+			]
+
 			// Search Box
 			+ SHorizontalBox::Slot()
 			.Padding(0, 7, 0, 7)
@@ -308,7 +352,7 @@ void SPluginBrowser::Tick(const FGeometry& AllottedGeometry, const double InCurr
 
 EVisibility SPluginBrowser::HandleRestartEditorNoticeVisibility() const
 {
-	return FPluginBrowserModule::Get().HasPluginsPendingEnable() ? EVisibility::Visible : EVisibility::Collapsed;
+	return FPluginBrowserModule::Get().ShowPendingRestart() ? EVisibility::Visible : EVisibility::Collapsed;
 }
 
 FReply SPluginBrowser::HandleRestartEditorButtonClicked() const
diff --git a/Engine/Plugins/Editor/PluginBrowser/Source/PluginBrowser/Private/SPluginCategoryTree.cpp b/Engine/Plugins/Editor/PluginBrowser/Source/PluginBrowser/Private/SPluginCategoryTree.cpp
--- a/Engine/Plugins/Editor/PluginBrowser/Source/PluginBrowser/Private/SPluginCategoryTree.cpp
+++ b/Engine/Plugins/Editor/PluginBrowser/Source/PluginBrowser/Private/SPluginCategoryTree.cpp
@@ -3,6 +3,7 @@
 #include "SPluginCategoryTree.h"
 #include "Framework/Views/TableViewMetadata.h"
 #include "Interfaces/IPluginManager.h"
+#include "Misc/ConfigCacheIni.h"
 #include "SPluginCategory.h"
 #include "SPluginBrowser.h"
 #include "Widgets/Views/STreeView.h"
@@ -226,6 +227,20 @@ void SPluginCategoryTree::RebuildAndFilterCategoryTree()
 		RootCategories.Insert(AllCategory, 0);
 	}
 
+	// Optionally hide some categories
+	{
+		TArray<FString> HiddenCategoryNames;
+		GConfig->GetArray(TEXT("EditorSettings"), TEXT("HidePluginCategoriesFromBrowser"), HiddenCategoryNames, GEditorIni);
+
+		for (TArray<TSharedPtr<FPluginCategory>>::TIterator CategoryIter = RootCategories.CreateIterator(); CategoryIter; ++CategoryIter)
+		{
+			if (HiddenCategoryNames.Contains((*CategoryIter)->Name))
+			{
+				CategoryIter.RemoveCurrent();
+			}
+		}
+	}
+
 	// Sort every single category alphabetically
 	for(TSharedPtr<FPluginCategory> RootCategory: RootCategories)
 	{
diff --git a/Engine/Plugins/Editor/PluginBrowser/Source/PluginBrowser/Private/SPluginPaths.cpp b/Engine/Plugins/Editor/PluginBrowser/Source/PluginBrowser/Private/SPluginPaths.cpp
new file mode 100644
--- /dev/null
+++ b/Engine/Plugins/Editor/PluginBrowser/Source/PluginBrowser/Private/SPluginPaths.cpp
@@ -0,0 +1,382 @@
+// Copyright Epic Games, Inc. All Rights Reserved.
+
+#include "SPluginPaths.h"
+#include "Interfaces/IPluginManager.h"
+#include "Interfaces/IProjectManager.h"
+#include "IStructureDetailsView.h"
+#include "Logging//StructuredLog.h"
+#include "Misc/ConfigCacheIni.h"
+#include "Modules/ModuleManager.h"
+#include "PluginBrowserModule.h"
+#include "PropertyEditorModule.h"
+#include "SSettingsEditorCheckoutNotice.h"
+#include "Styling/StyleColors.h"
+#include "Widgets/Layout/SBox.h"
+#include "Widgets/Text/STextBlock.h"
+
+#define LOCTEXT_NAMESPACE "SPluginPaths"
+
+
+namespace
+{
+	static const FName ProjectPathsName = GET_MEMBER_NAME_CHECKED(FPluginPaths_External, AdditionalPluginDirectories);
+	static const FName UserPathsName = GET_MEMBER_NAME_CHECKED(FPluginPaths_External, UserPluginDirectories);
+	static const FName CmdLinePathsName = GET_MEMBER_NAME_CHECKED(FPluginPaths_External, CommandLineDirectories);
+	static const FName EnvironmentPathsName = GET_MEMBER_NAME_CHECKED(FPluginPaths_External, EnvironmentDirectories);
+}
+
+
+SLATE_IMPLEMENT_WIDGET(SPluginPaths)
+
+
+void SPluginPaths::PrivateRegisterAttributes(FSlateAttributeInitializer& InAttributeInitializer)
+{
+}
+
+
+void SPluginPaths::Construct(const FArguments& InArgs)
+{
+	ExternalPathsStruct = MakeShared<TStructOnScope<FPluginPaths_External>>();
+	ExternalPathsStruct->InitializeAs<FPluginPaths_External>();
+	FPluginPaths_External& ExternalPaths = GetExternalPaths();
+
+	IPluginManager& PluginManager = IPluginManager::Get();
+
+	// Categorize configured external paths.
+	TSet<FExternalPluginPath> AllExternalPaths;
+	PluginManager.GetExternalPluginSources(AllExternalPaths);
+
+	for (const FExternalPluginPath& ExternalPath : AllExternalPaths)
+	{
+		switch (ExternalPath.Source)
+		{
+		case EPluginExternalSource::ProjectDescriptor:
+			ExternalPaths.AdditionalPluginDirectories.Emplace(ExternalPath.Path);
+			break;
+		case EPluginExternalSource::Other:
+			ExternalPaths.UserPluginDirectories.Emplace(ExternalPath.Path);
+			break;
+		case EPluginExternalSource::CommandLine:
+			ExternalPaths.CommandLineDirectories.Emplace(ExternalPath.Path);
+			break;
+		case EPluginExternalSource::Environment:
+			ExternalPaths.EnvironmentDirectories.Emplace(ExternalPath.Path);
+			break;
+		default:
+			ensureAlways("Unhandled external plugin source");
+		}
+	}
+
+	FPropertyEditorModule& PropertyEditorModule =
+		FModuleManager::GetModuleChecked<FPropertyEditorModule>("PropertyEditor");
+
+	FDetailsViewArgs DetailsViewArgs;
+	DetailsViewArgs.NotifyHook = this;
+	DetailsViewArgs.NameAreaSettings = FDetailsViewArgs::HideNameArea;
+	DetailsViewArgs.bAllowSearch = false;
+	DetailsViewArgs.bHideSelectionTip = true;
+	DetailsViewArgs.bShowPropertyMatrixButton = false;
+	DetailsViewArgs.bShowModifiedPropertiesOption = false;
+	DetailsViewArgs.bShowHiddenPropertiesWhilePlayingOption = false;
+	DetailsViewArgs.bAllowMultipleTopLevelObjects = true;
+	DetailsViewArgs.bShowKeyablePropertiesOption = false;
+	DetailsViewArgs.bShowAnimatedPropertiesOption = false;
+
+	// Setting these to true is necessary to specify our own Visible delegate,
+	// which is otherwise overwritten by FStructureDetailsViewFilter.
+	FStructureDetailsViewArgs StructViewArgs;
+	StructViewArgs.bShowObjects = true;
+	StructViewArgs.bShowInterfaces = true;
+
+	ExternalView = PropertyEditorModule.CreateStructureDetailView(DetailsViewArgs, StructViewArgs, ExternalPathsStruct);
+	ExternalView->GetDetailsView()->SetIsPropertyReadOnlyDelegate(FIsPropertyReadOnly::CreateSP(this, &SPluginPaths::HandleIsPropertyReadOnly));
+	ExternalView->GetDetailsView()->SetIsPropertyVisibleDelegate(FIsPropertyReadOnly::CreateSP(this, &SPluginPaths::HandleIsPropertyVisible));
+	ExternalView->GetDetailsView()->ForceRefresh();
+
+	TSharedPtr<SWidget> CheckoutWidget;
+	if (CanModifyProjectPaths())
+	{
+		CheckoutWidget = SNew(SSettingsEditorCheckoutNotice)
+			.ConfigFilePath(FPaths::GetProjectFilePath());
+	}
+
+	ChildSlot
+	[
+		SNew(SVerticalBox)
+		+ SVerticalBox::Slot()
+		.AutoHeight()
+		[
+			SNew(SBox)
+			.Padding(24.f, 16.f)
+			[
+				SNew(SVerticalBox)
+				+ SVerticalBox::Slot()
+				.AutoHeight()
+				[
+					// Category title
+					SNew(STextBlock)
+					.Font(FAppStyle::GetFontStyle("SettingsEditor.CatgoryAndSectionFont"))
+					.Text(LOCTEXT("SettingsTitle", "External Plugin Directories"))
+				]
+				+ SVerticalBox::Slot()
+				.AutoHeight()
+				.Padding(0.f, 5.f, 0.f, 0.f)
+				[
+					// Category description
+					SNew(STextBlock)
+					.ColorAndOpacity(FSlateColor::UseSubduedForeground())
+					.Text(LOCTEXT("SettingsDescription", "Configure additional locations which should be enumerated for plugins."))
+				]
+			]
+		]
+		+ SVerticalBox::Slot()
+		.AutoHeight()
+		.Padding(0.f, 0.f, 8.f, 0.f)
+		[
+			CheckoutWidget ? CheckoutWidget.ToSharedRef() : SNullWidget::NullWidget
+		]
+		+ SVerticalBox::Slot()
+		.AutoHeight()
+		[
+			ExternalView->GetWidget().ToSharedRef()
+		]
+	];
+}
+
+
+void SPluginPaths::NotifyPreChange(class FEditPropertyChain* PropertyAboutToChange)
+{
+	FNotifyHook::NotifyPreChange(PropertyAboutToChange);
+
+	// Update the cached "before" copy for subsequent diff in *post* change notification.
+	PreviousExternalPathsForDiff = GetExternalPaths();
+}
+
+void SPluginPaths::NotifyPostChange(const FPropertyChangedEvent& InPropertyChangedEvent, FEditPropertyChain* InPropertyThatChanged)
+{
+	FNotifyHook::NotifyPostChange(InPropertyChangedEvent, InPropertyThatChanged);
+
+	if (InPropertyThatChanged == nullptr)
+	{
+		return;
+	}
+
+	const FName ActiveMemberName = InPropertyThatChanged->GetActiveMemberNode()->GetValue()->GetFName();
+
+	bool bPluginDirectoriesChanged = false;
+
+	if (ActiveMemberName == ProjectPathsName)
+	{
+		const FString ProjectPath = FPaths::GetProjectFilePath();
+		if (!SettingsHelpers::CheckOutOrAddFile(ProjectPath))
+		{
+			SettingsHelpers::MakeWritable(ProjectPath);
+		}
+
+		TSet<FString> AddedPaths, RemovedPaths;
+		DiffAddedAndRemoved(
+			PreviousExternalPathsForDiff.AdditionalPluginDirectories,
+			GetExternalPaths().AdditionalPluginDirectories,
+			AddedPaths,
+			RemovedPaths
+		);
+
+		if (AddedPaths.Num())
+		{
+			UE_LOGFMT(LogPluginBrowser, Display, "Added project plugin directories: {AddedPaths}",
+				FString::Join(AddedPaths, TEXT("; ")));
+		}
+
+		if (RemovedPaths.Num())
+		{
+			UE_LOGFMT(LogPluginBrowser, Display, "Removed project plugin directories: {RemovedPaths}",
+				FString::Join(RemovedPaths, TEXT("; ")));
+		}
+
+		IProjectManager& ProjectManager = IProjectManager::Get();
+		for (const FString& RemovedPath : RemovedPaths)
+		{
+			bPluginDirectoriesChanged |= ProjectManager.UpdateAdditionalPluginDirectory(RemovedPath, false);
+		}
+		for (const FString& AddedPath : AddedPaths)
+		{
+			bPluginDirectoriesChanged |= ProjectManager.UpdateAdditionalPluginDirectory(AddedPath, true);
+		}
+	}
+	else if (ActiveMemberName == UserPathsName)
+	{
+		TSet<FString> AddedPaths, RemovedPaths;
+		DiffAddedAndRemoved(
+			PreviousExternalPathsForDiff.UserPluginDirectories,
+			GetExternalPaths().UserPluginDirectories,
+			AddedPaths,
+			RemovedPaths
+		);
+
+		if (AddedPaths.Num())
+		{
+			UE_LOGFMT(LogPluginBrowser, Display, "Added user plugin directories: {AddedPaths}",
+				FString::Join(AddedPaths, TEXT("; ")));
+		}
+
+		if (RemovedPaths.Num())
+		{
+			UE_LOGFMT(LogPluginBrowser, Display, "Removed user plugin directories: {RemovedPaths}",
+				FString::Join(RemovedPaths, TEXT("; ")));
+		}
+
+		IPluginManager& PluginManager = IPluginManager::Get();
+		for (const FString& RemovedPath : RemovedPaths)
+		{
+			bPluginDirectoriesChanged |= PluginManager.RemovePluginSearchPath(RemovedPath, false);
+		}
+		for (const FString& AddedPath : AddedPaths)
+		{
+			bPluginDirectoriesChanged |= PluginManager.AddPluginSearchPath(AddedPath, false);
+		}
+
+		if (bPluginDirectoriesChanged)
+		{
+			PluginManager.RefreshPluginsList();
+		}
+	}
+
+	if (bPluginDirectoriesChanged)
+	{
+		FPluginBrowserModule::Get().OnPluginDirectoriesChanged().Broadcast();
+	}
+}
+
+
+FPluginPaths_External& SPluginPaths::GetExternalPaths()
+{
+	check(ExternalPathsStruct);
+	return *ExternalPathsStruct->Get();
+}
+
+
+const FPluginPaths_External& SPluginPaths::GetExternalPaths() const
+{
+	check(ExternalPathsStruct);
+	return *ExternalPathsStruct->Get();
+}
+
+
+bool SPluginPaths::CanModifyProjectPaths() const
+{
+	// If the new setting has been specified explicitly, use that.
+	bool bCanModifyProject = true;
+	if (GConfig->GetBool(TEXT("EditorSettings"), TEXT("bCanModifyProjectPluginDirectoriesFromBrowser"), bCanModifyProject, GEditorIni))
+	{
+		return bCanModifyProject;
+	}
+
+	// Fall back to whether plugins can be enabled/disabled in general.
+	bool bModifyPluginsEnabled = true;
+	GConfig->GetBool(TEXT("EditorSettings"), TEXT("bCanModifyPluginsFromBrowser"), bModifyPluginsEnabled, GEditorIni);
+	return bModifyPluginsEnabled;
+}
+
+
+bool SPluginPaths::CanModifyUserPaths() const
+{
+	bool bCanModifyExternal = false;
+	GConfig->GetBool(TEXT("EditorSettings"), TEXT("bCanModifyUserPluginDirectoriesFromBrowser"), bCanModifyExternal, GEditorIni);
+	return bCanModifyExternal;
+}
+
+
+bool SPluginPaths::ShouldShowProjectPaths() const
+{
+	return CanModifyProjectPaths() || (GetExternalPaths().AdditionalPluginDirectories.Num() > 0);
+}
+
+
+bool SPluginPaths::ShouldShowUserPaths() const
+{
+	return CanModifyUserPaths() || (GetExternalPaths().UserPluginDirectories.Num() > 0);
+}
+
+
+bool SPluginPaths::HandleIsPropertyReadOnly(const FPropertyAndParent& InPropertyAndParent) const
+{
+	const FName PropertyName = InPropertyAndParent.Property.GetFName();
+	const FName ParentName = (InPropertyAndParent.ParentProperties.Num() > 0)
+		? InPropertyAndParent.ParentProperties[0]->GetFName() : NAME_None;
+
+	if (PropertyName == ProjectPathsName || ParentName == ProjectPathsName)
+	{
+		return !CanModifyProjectPaths();
+	}
+
+	if (PropertyName == UserPathsName || ParentName == UserPathsName)
+	{
+		return !CanModifyUserPaths();
+	}
+
+	if (ParentName == CmdLinePathsName)
+	{
+		return true;
+	}
+
+	if (ParentName == EnvironmentPathsName)
+	{
+		return true;
+	}
+
+	return false;
+}
+
+
+bool SPluginPaths::HandleIsPropertyVisible(const FPropertyAndParent& InPropertyAndParent) const
+{
+	const FName PropertyName = InPropertyAndParent.Property.GetFName();
+	const FName ParentName = (InPropertyAndParent.ParentProperties.Num() > 0)
+		? InPropertyAndParent.ParentProperties[0]->GetFName() : NAME_None;
+
+	if (PropertyName == ProjectPathsName || ParentName == ProjectPathsName)
+	{
+		return ShouldShowProjectPaths();
+	}
+
+	if (PropertyName == UserPathsName || ParentName == UserPathsName)
+	{
+		return ShouldShowUserPaths();
+	}
+
+	return true;
+}
+
+
+void SPluginPaths::DiffAddedAndRemoved(
+	const TArray<FDirectoryPath>& InBefore,
+	const TArray<FDirectoryPath>& InAfter,
+	TSet<FString>& OutAdded,
+	TSet<FString>& OutRemoved)
+{
+	TSet<FString> BeforeSet;
+	BeforeSet.Reserve(InBefore.Num());
+	for (const FDirectoryPath& BeforePath : InBefore)
+	{
+		if (!BeforePath.Path.IsEmpty())
+		{
+			BeforeSet.Add(BeforePath.Path);
+		}
+	}
+
+	TSet<FString> AfterSet;
+	AfterSet.Reserve(InAfter.Num());
+	for (const FDirectoryPath& AfterPath : InAfter)
+	{
+		if (!AfterPath.Path.IsEmpty())
+		{
+			AfterSet.Add(AfterPath.Path);
+		}
+	}
+
+	OutAdded = AfterSet.Difference(BeforeSet);
+	OutRemoved = BeforeSet.Difference(AfterSet);
+}
+
+
+#undef LOCTEXT_NAMESPACE
diff --git a/Engine/Plugins/Editor/PluginBrowser/Source/PluginBrowser/Private/SPluginPaths.h b/Engine/Plugins/Editor/PluginBrowser/Source/PluginBrowser/Private/SPluginPaths.h
new file mode 100644
--- /dev/null
+++ b/Engine/Plugins/Editor/PluginBrowser/Source/PluginBrowser/Private/SPluginPaths.h
@@ -0,0 +1,93 @@
+// Copyright Epic Games, Inc. All Rights Reserved.
+
+#pragma once
+
+#include "Misc/NotifyHook.h"
+#include "UObject/StructOnScope.h"
+#include "Widgets/DeclarativeSyntaxSupport.h"
+#include "Widgets/SCompoundWidget.h"
+#include "SPluginPaths.generated.h"
+
+
+class FEditPropertyChain;
+struct FPropertyAndParent;
+struct FPropertyChangedEvent;
+class IStructureDetailsView;
+
+
+/** Struct used as a model to expose relevant lists via details view. */
+USTRUCT()
+struct FPluginPaths_External
+{
+	GENERATED_BODY()
+
+	UPROPERTY(Category="Plugin Directories", EditAnywhere, meta=(
+		ToolTip="Stored in the .uproject descriptor."))
+	TArray<FDirectoryPath> AdditionalPluginDirectories;
+
+	UPROPERTY(Category="Plugin Directories", EditAnywhere)
+	TArray<FDirectoryPath> UserPluginDirectories;
+
+	UPROPERTY(Category="Plugin Directories", VisibleAnywhere, EditFixedSize, meta=(NoResetToDefault,
+		ToolTip="Specified via the -PLUGIN= command line switch; cannot be modified here."))
+	TArray<FDirectoryPath> CommandLineDirectories;
+
+	UPROPERTY(Category="Plugin Directories", VisibleAnywhere, EditFixedSize, meta=(NoResetToDefault,
+		ToolTip="Specified via the UE_ADDITIONAL_PLUGIN_PATHS environment variable; cannot be modified here."))
+	TArray<FDirectoryPath> EnvironmentDirectories;
+};
+
+
+/** Widget that marshals details view to/from project and plugin manager. */
+class SPluginPaths
+	: public SCompoundWidget
+	, public FNotifyHook
+{
+	SLATE_DECLARE_WIDGET(SPluginPaths, SCompoundWidget)
+
+public:
+	SLATE_BEGIN_ARGS(SPluginPaths) { }
+	SLATE_END_ARGS()
+
+public:
+	void Construct(const FArguments& InArgs);
+
+	//~ BEGIN FNotifyHook interface
+	virtual void NotifyPreChange(FEditPropertyChain* PropertyAboutToChange) override;
+	virtual void NotifyPostChange(const FPropertyChangedEvent& PropertyChangedEvent, FEditPropertyChain* PropertyThatChanged) override;
+	//~ END FNotifyHook interface
+
+private:
+	/** Mutable convenience getter for ExternalPathsStruct TStructOnScope. */
+	FPluginPaths_External& GetExternalPaths();
+	/** Const convenience getter for ExternalPathsStruct TStructOnScope. */
+	const FPluginPaths_External& GetExternalPaths() const;
+
+	bool CanModifyProjectPaths() const;
+	bool CanModifyUserPaths() const;
+	bool ShouldShowProjectPaths() const;
+	bool ShouldShowUserPaths() const;
+
+	/** Details view read-only delegate. */
+	bool HandleIsPropertyReadOnly(const FPropertyAndParent& InPropertyAndParent) const;
+
+	/** Details view visibility delegate. */
+	bool HandleIsPropertyVisible(const FPropertyAndParent& InPropertyAndParent) const;
+
+	/** Given two FDirectoryPath arrays, indicate what path strings were added/removed from InAfter compared to InBefore. */
+	void DiffAddedAndRemoved(
+		const TArray<FDirectoryPath>& InBefore,
+		const TArray<FDirectoryPath>& InAfter,
+		TSet<FString>& OutAdded,
+		TSet<FString>& OutRemoved
+	);
+
+	/** The model being edited in the details view. */
+	TSharedPtr<TStructOnScope<FPluginPaths_External>> ExternalPathsStruct;
+
+	/** The primary editor widget. */
+	TSharedPtr<IStructureDetailsView> ExternalView;
+
+	/** Cache that's updated in NotifyPreChange, and diffed against in NotifyPostChange. */
+	FPluginPaths_External PreviousExternalPathsForDiff;
+};
diff --git a/Engine/Plugins/Editor/PluginBrowser/Source/PluginBrowser/Private/SPluginTileList.cpp b/Engine/Plugins/Editor/PluginBrowser/Source/PluginBrowser/Private/SPluginTileList.cpp
--- a/Engine/Plugins/Editor/PluginBrowser/Source/PluginBrowser/Private/SPluginTileList.cpp
+++ b/Engine/Plugins/Editor/PluginBrowser/Source/PluginBrowser/Private/SPluginTileList.cpp
@@ -137,7 +137,7 @@ void SPluginTileList::SetNeedsRefresh()
 FVector2D SPluginTileList::GetListBorderFadeDistance() const
 {
 	// Negative fade distance when there is no restart editor warning to make the shadow disappear
-	FVector2D ReturnVal = FPluginBrowserModule::Get().HasPluginsPendingEnable() ? FVector2D(0.01f, 0.01f) : FVector2D(-1.0f, -1.0f);
+	FVector2D ReturnVal = FPluginBrowserModule::Get().ShowPendingRestart() ? FVector2D(0.01f, 0.01f) : FVector2D(-1.0f, -1.0f);
 
 	return ReturnVal;
 }
diff --git a/Engine/Plugins/Editor/PluginBrowser/Source/PluginBrowser/Public/IPluginBrowser.h b/Engine/Plugins/Editor/PluginBrowser/Source/PluginBrowser/Public/IPluginBrowser.h
--- a/Engine/Plugins/Editor/PluginBrowser/Source/PluginBrowser/Public/IPluginBrowser.h
+++ b/Engine/Plugins/Editor/PluginBrowser/Source/PluginBrowser/Public/IPluginBrowser.h
@@ -6,6 +6,7 @@
 #include "Features/IPluginsEditorFeature.h"
 
 DECLARE_DELEGATE_OneParam(FOnLaunchReferenceViewer, TSharedPtr<IPlugin>);
+DECLARE_MULTICAST_DELEGATE(FOnPluginDirectoriesChanged);
 
 /**
  * The public interface to this module
@@ -50,6 +51,9 @@ public:
 	 * Allows another module that supplies the plugin reference viewer to launch when requested.
 	 */
 	virtual FOnLaunchReferenceViewer& OnLaunchReferenceViewerDelegate() = 0;
+
+	/** Called when the external plugin directories configuration is modified via the browser. */
+	virtual FOnPluginDirectoriesChanged& OnPluginDirectoriesChanged() = 0;
 };
 
 
diff --git a/Engine/Source/Runtime/Projects/Private/PluginManager.cpp b/Engine/Source/Runtime/Projects/Private/PluginManager.cpp
--- a/Engine/Source/Runtime/Projects/Private/PluginManager.cpp
+++ b/Engine/Source/Runtime/Projects/Private/PluginManager.cpp
@@ -100,20 +100,15 @@ namespace PluginSystemDefs
 	static const TCHAR* PluginDescriptorFileExtension = TEXT( ".uplugin" );
 
 	/**
-	 * Parsing the command line and loads any foreign plugins that were
-	 * specified using the -PLUGIN= command.
+	 * Append all paths specified using the -PLUGIN= switch.
 	 *
-	 * @param  CommandLine    The commandline used to launch the editor.
-	 * @param  SearchPathsOut 
-	 * @return The number of plugins that were specified using the -PLUGIN param.
+	 * @param  PluginPathsOut Has a value added for each -PLUGIN=value pair.
 	 */
-	static int32 GetAdditionalPluginPaths(TSet<FString>& PluginPathsOut)
+	static void GetCmdLinePluginPaths(TSet<FExternalPluginPath>& PluginPathsOut)
 	{
 		const TCHAR* SwitchStr = TEXT("PLUGIN=");
 		const int32  SwitchLen = FCString::Strlen(SwitchStr);
 
-		int32 PluginCount = 0;
-
 		const TCHAR* SearchStr = FCommandLine::Get();
 		do
 		{
@@ -123,9 +118,7 @@ namespace PluginSystemDefs
 			if (FParse::Value(SearchStr, SwitchStr, PluginPath))
 			{
 				FString PluginDir = FPaths::GetPath(PluginPath);
-				PluginPathsOut.Add(PluginDir);
-
-				++PluginCount;
+				PluginPathsOut.Add(FExternalPluginPath{ PluginDir, EPluginExternalSource::CommandLine });
 				SearchStr += SwitchLen + PluginPath.Len();
 			}
 			else
@@ -133,13 +126,20 @@ namespace PluginSystemDefs
 				break;
 			}
 		} while (SearchStr != nullptr);
+	}
 
+	/**
+	 * Retrieve a list of all paths specified via the UE_ADDITIONAL_PLUGIN_PATHS environment variable.
+	 *
+	 * @param  PluginPathsOut A list of values split from the environment variable.
+	 */
+	static void GetEnvPluginPaths(TSet<FExternalPluginPath>& PluginPathsOut)
+	{
 		TArray<FString> AdditionalEnvPaths = UE::PluginManager::Private::GetAdditionalExternalPluginsByEnvVar();
 		for (const FString& Path : AdditionalEnvPaths)
 		{
-			PluginPathsOut.Add(Path);
+			PluginPathsOut.Add(FExternalPluginPath{ Path, EPluginExternalSource::Environment });
 		}
-		return PluginCount;
 	}
 
 
@@ -157,7 +157,7 @@ namespace PluginSystemDefs
  * Set of simple (ideally inlinable) helper methods intended to obscure how 
  * `FDiscoveredPluginMap` is implemented (so that we can more easily change out it's type as needed).
  * 
- * Also dictactes how we separate the one "offered"  plugin out from other versions of the same plugin 
+ * Also dictates how we separate the one "offered"  plugin out from other versions of the same plugin 
  * (see DiscoveredPluginMapUtils::EInsertionType).
  */
 namespace DiscoveredPluginMapUtils
@@ -752,7 +752,13 @@ void FPluginManager::DiscoverAllPlugins()
 {
 	ensure( AllPlugins.Num() == 0 );		// Should not have already been initialized!
 
-	PluginSystemDefs::GetAdditionalPluginPaths(PluginDiscoveryPaths);
+	PluginSystemDefs::GetCmdLinePluginPaths(ExternalPluginSources);
+	PluginSystemDefs::GetEnvPluginPaths(ExternalPluginSources);
+	for (const FExternalPluginPath& ExternalSource : ExternalPluginSources)
+	{
+		PluginDiscoveryPaths.Add(ExternalSource.Path);
+	}
+
 	ReadAllPlugins(AllPlugins, PluginDiscoveryPaths);
 
 	PluginsToConfigure.Reserve(AllPlugins.Num());
@@ -3173,20 +3179,52 @@ TSharedPtr<IPlugin> FPluginManager::GetModuleOwnerPlugin(FName ModuleName) const
 
 bool FPluginManager::AddPluginSearchPath(const FString& ExtraDiscoveryPath, bool bRefresh)
 {
+	const FString FullPath = FPaths::ConvertRelativePathToFull(ExtraDiscoveryPath);
+
 	bool bAlreadyExists = false;
-	PluginDiscoveryPaths.Add(FPaths::ConvertRelativePathToFull(ExtraDiscoveryPath), &bAlreadyExists);
+	ExternalPluginSources.Add(FExternalPluginPath{ FullPath, EPluginExternalSource::Other });
+	PluginDiscoveryPaths.Add(FullPath, &bAlreadyExists);
+
 	if (bRefresh)
 	{
 		RefreshPluginsList();
 	}
+
 	return !bAlreadyExists;
 }
 
+bool FPluginManager::RemovePluginSearchPath(const FString& PathToRemove, bool bRefresh)
+{
+	const FString FullPath = FPaths::ConvertRelativePathToFull(PathToRemove);
+
+	const bool bPathWasPresent = PluginDiscoveryPaths.Contains(FullPath);
+
+	ExternalPluginSources.Remove(FExternalPluginPath{ FullPath, EPluginExternalSource::Other });
+	PluginDiscoveryPaths.Remove(FullPath);
+
+	if (bRefresh)
+	{
+		RefreshPluginsList();
+	}
+
+	return bPathWasPresent;
+}
+
 const TSet<FString>& FPluginManager::GetAdditionalPluginSearchPaths() const
 {
 	return PluginDiscoveryPaths;
 }
 
+void FPluginManager::GetExternalPluginSources(TSet<FExternalPluginPath>& OutPluginSources) const
+{
+	OutPluginSources.Append(ExternalPluginSources);
+
+	for (const FString& ProjectPath : IProjectManager::Get().GetAdditionalPluginDirectories())
+	{
+		OutPluginSources.Add(FExternalPluginPath{ ProjectPath, EPluginExternalSource::ProjectDescriptor });
+	}
+}
+
 TArray<TSharedRef<IPlugin>> FPluginManager::GetPluginsWithPakFile() const
 {
 	return PluginsWithPakFile;
diff --git a/Engine/Source/Runtime/Projects/Private/PluginManager.h b/Engine/Source/Runtime/Projects/Private/PluginManager.h
--- a/Engine/Source/Runtime/Projects/Private/PluginManager.h
+++ b/Engine/Source/Runtime/Projects/Private/PluginManager.h
@@ -204,7 +204,9 @@ public:
 #endif //WITH_EDITOR
 
 	virtual bool AddPluginSearchPath(const FString& ExtraDiscoveryPath, bool bRefresh = true) override;
-	const TSet<FString>& GetAdditionalPluginSearchPaths() const override;
+	virtual bool RemovePluginSearchPath(const FString& PathToRemove, bool bRefresh = true) override;
+	virtual const TSet<FString>& GetAdditionalPluginSearchPaths() const override;
+	virtual void GetExternalPluginSources(TSet<FExternalPluginPath>& OutPluginSources) const override;
 	virtual TArray<TSharedRef<IPlugin>> GetPluginsWithPakFile() const override;
 	virtual FNewPluginMountedEvent& OnNewPluginCreated() override;
 	virtual FNewPluginMountedEvent& OnNewPluginMounted() override;
@@ -358,6 +360,9 @@ private:
 	bool bPreloadedBinaries = false;
 
 	/** List of additional directory paths to search for plugins within */
+	TSet<FExternalPluginPath> ExternalPluginSources;
+
+	/** Projection of ExternalPluginSources paths only; maintained for backward compatibility */
 	TSet<FString> PluginDiscoveryPaths;
 
 	/** Callback for notifications that a new plugin was mounted */
diff --git a/Engine/Source/Runtime/Projects/Public/Interfaces/IPluginManager.h b/Engine/Source/Runtime/Projects/Public/Interfaces/IPluginManager.h
--- a/Engine/Source/Runtime/Projects/Public/Interfaces/IPluginManager.h
+++ b/Engine/Source/Runtime/Projects/Public/Interfaces/IPluginManager.h
@@ -45,6 +45,42 @@ enum class EPluginType
 	Mod,
 };
 
+/**
+ * Enum describing a source of EPluginType::External plugins
+ */
+enum class EPluginExternalSource
+{
+	/** .uproject `AdditionalPluginDirectories` field */
+	ProjectDescriptor,
+
+	/** Passed via `-plugin=` command line switch. */
+	CommandLine,
+
+	/** Inherited from `UE_ADDITIONAL_PLUGIN_PATHS` environment variable. */
+	Environment,
+
+	/** None of the above (e.g. added explicitly by external code). */
+	Other,
+};
+
+/**
+ * Combination of a directory searched for external plugins, and its configuration source.
+ */
+struct FExternalPluginPath
+{
+	FString Path;
+	EPluginExternalSource Source;
+
+	inline bool operator==(const FExternalPluginPath& Other) const
+	{
+		return (Other.Path == Path) && (Other.Source == Source);
+	}
+
+	friend uint32 GetTypeHash(const FExternalPluginPath& ExternalPluginPath)
+	{
+		return HashCombine(GetTypeHash(ExternalPluginPath.Path), GetTypeHash(ExternalPluginPath.Source));
+	}
+};
 
 /**
  * Simple data structure that is filled when querying information about plug-ins.
@@ -421,11 +457,28 @@ public:
 	 */
 	virtual bool AddPluginSearchPath(const FString& ExtraDiscoveryPath, bool bRefresh = true) = 0;
 
+	/**
+	 * Removes the specified path from consideration for available plugins.
+	 * Optionally refreshes the manager after the path has been removed.
+	 *
+	 * @param  PathToRemove			The path no longer searched for additional plugins.
+	 * @param  bRefresh				Signals the function to refresh the plugin database after the path has been removed
+	 * @return Whether the plugin search path was modified
+	 */
+	virtual bool RemovePluginSearchPath(const FString& PathToRemove, bool bRefresh = true) = 0;
+
 	/**
 	 * Returns the list of extra directories that are recursively searched for plugins (aside from the engine and project plugin directories).
+	 * NOTE: You may also want to check IProjectManager::GetAdditionalPluginDirectories, which are not included here!
 	 */
 	virtual const TSet<FString>& GetAdditionalPluginSearchPaths() const = 0;
 
+	/**
+	 * Similar to GetAdditionalPluginSearchPaths, but also returns configuration sources for each path.
+	 * Unlike GetAdditionalPluginSearchPaths, this DOES include AdditionalPluginDirectories from the uproject descriptor.
+	 */
+	virtual void GetExternalPluginSources(TSet<FExternalPluginPath>& OutPluginSources) const = 0;
+
 	/**
 	 * Gets an array of plugins that loaded their own content pak file
 	 */
