Plugin Browser: Mac build fix. Followup to CL 35747837.

#rnx

diff --git a/Engine/Plugins/Editor/PluginBrowser/Source/PluginBrowser/Private/SPluginPaths.cpp b/Engine/Plugins/Editor/PluginBrowser/Source/PluginBrowser/Private/SPluginPaths.cpp
--- a/Engine/Plugins/Editor/PluginBrowser/Source/PluginBrowser/Private/SPluginPaths.cpp
+++ b/Engine/Plugins/Editor/PluginBrowser/Source/PluginBrowser/Private/SPluginPaths.cpp
@@ -51,16 +51,16 @@ void SPluginPaths::Construct(const FArguments& InArgs)
 		switch (ExternalPath.Source)
 		{
 		case EPluginExternalSource::ProjectDescriptor:
-			ExternalPaths.AdditionalPluginDirectories.Emplace(ExternalPath.Path);
+			ExternalPaths.AdditionalPluginDirectories.Add(FDirectoryPath{ ExternalPath.Path });
 			break;
 		case EPluginExternalSource::Other:
-			ExternalPaths.UserPluginDirectories.Emplace(ExternalPath.Path);
+			ExternalPaths.UserPluginDirectories.Add(FDirectoryPath{ ExternalPath.Path });
 			break;
 		case EPluginExternalSource::CommandLine:
-			ExternalPaths.CommandLineDirectories.Emplace(ExternalPath.Path);
+			ExternalPaths.CommandLineDirectories.Add(FDirectoryPath{ ExternalPath.Path });
 			break;
 		case EPluginExternalSource::Environment:
-			ExternalPaths.EnvironmentDirectories.Emplace(ExternalPath.Path);
+			ExternalPaths.EnvironmentDirectories.Add(FDirectoryPath{ ExternalPath.Path });
 			break;
 		default:
 			ensureAlways("Unhandled external plugin source");
