Plugin Browser: Display "External" plugins in a separate category.

diff --git a/Engine/Plugins/Editor/PluginBrowser/Source/PluginBrowser/Private/SPluginCategoryTree.cpp b/Engine/Plugins/Editor/PluginBrowser/Source/PluginBrowser/Private/SPluginCategoryTree.cpp
--- a/Engine/Plugins/Editor/PluginBrowser/Source/PluginBrowser/Private/SPluginCategoryTree.cpp
+++ b/Engine/Plugins/Editor/PluginBrowser/Source/PluginBrowser/Private/SPluginCategoryTree.cpp
@@ -21,6 +21,7 @@ void SPluginCategoryTree::Construct( const FArguments& Args, const TSharedRef< S
 	// Create the root categories
 	AllCategory = MakeShareable(new FPluginCategory(NULL, TEXT("All"), LOCTEXT("AllCategoryName", "All Plugins")));
 	BuiltInCategory = MakeShareable(new FPluginCategory(NULL, TEXT("Built-In"), LOCTEXT("BuiltInCategoryName", "Built-In")));
+	ExternalCategory = MakeShareable(new FPluginCategory(NULL, TEXT("External"), LOCTEXT("ExternalCategoryName", "External")));
 	InstalledCategory = MakeShareable(new FPluginCategory(NULL, TEXT("Installed"), LOCTEXT("InstalledCategoryName", "Installed")));
 	ProjectCategory = MakeShareable(new FPluginCategory(NULL, TEXT("Project"), LOCTEXT("ProjectCategoryName", "Project")));
 	ModCategory = MakeShareable(new FPluginCategory(NULL, TEXT("Mods"), LOCTEXT("ModsCategoryName", "Mods")));
@@ -114,6 +115,10 @@ void SPluginCategoryTree::RebuildAndFilterCategoryTree()
 		{
 			RootCategory = ModCategory;
 		}
+		else if (Plugin->GetType() == EPluginType::External)
+		{
+			RootCategory = ExternalCategory;
+		}
 		else if(Plugin->GetDescriptor().bInstalled)
 		{
 			RootCategory = InstalledCategory;
@@ -218,6 +223,10 @@ void SPluginCategoryTree::RebuildAndFilterCategoryTree()
 	{
 		RootCategories.Add(InstalledCategory);
 	}
+	if (ExternalCategory->SubCategories.Num() > 0 || ExternalCategory->Plugins.Num() > 0)
+	{
+		RootCategories.Add(ExternalCategory);
+	}
 	if(BuiltInCategory->SubCategories.Num() > 0 || BuiltInCategory->Plugins.Num() > 0)
 	{
 		RootCategories.Add(BuiltInCategory);
diff --git a/Engine/Plugins/Editor/PluginBrowser/Source/PluginBrowser/Private/SPluginCategoryTree.h b/Engine/Plugins/Editor/PluginBrowser/Source/PluginBrowser/Private/SPluginCategoryTree.h
--- a/Engine/Plugins/Editor/PluginBrowser/Source/PluginBrowser/Private/SPluginCategoryTree.h
+++ b/Engine/Plugins/Editor/PluginBrowser/Source/PluginBrowser/Private/SPluginCategoryTree.h
@@ -91,6 +91,9 @@ private:
 	/** Category for built-in plugins */
 	TSharedPtr<FPluginCategory> BuiltInCategory;
 
+	/** Category for external plugins */
+	TSharedPtr<FPluginCategory> ExternalCategory;
+
 	/** Category for installed plugins */
 	TSharedPtr<FPluginCategory> InstalledCategory;
 
