Update plugin browser to show disabled entries for plugins that require C++ when your project is blueprint only (rather than just hide them)

[RN]

diff --git a/Engine/Plugins/Editor/PluginBrowser/Source/PluginBrowser/Private/DefaultPluginWizardDefinition.cpp b/Engine/Plugins/Editor/PluginBrowser/Source/PluginBrowser/Private/DefaultPluginWizardDefinition.cpp
--- a/Engine/Plugins/Editor/PluginBrowser/Source/PluginBrowser/Private/DefaultPluginWizardDefinition.cpp
+++ b/Engine/Plugins/Editor/PluginBrowser/Source/PluginBrowser/Private/DefaultPluginWizardDefinition.cpp
@@ -40,20 +40,17 @@ void FDefaultPluginWizardDefinition::PopulateTemplatesSource()
 	ContentOnlyTemplate->SortPriority = 1;
 	TemplateDefinitions.Add(ContentOnlyTemplate);
 
-	if (!bIsContentOnlyProject)
-	{
-		// Insert the blank template to make sure it appears before the content only template.
+	// Insert the blank template to make sure it appears before the content only template.
 
-		TSharedRef<FPluginTemplateDescription> BlankTemplate = MakeShareable(new FPluginTemplateDescription(BlankTemplateName, BlankDescription, PluginBaseDir / TEXT("Templates") / TEXT("Blank"), true, EHostType::Runtime));
-		BlankTemplate->SortPriority = 2;
-		TemplateDefinitions.Add(BlankTemplate);
+	TSharedRef<FPluginTemplateDescription> BlankTemplate = MakeShareable(new FPluginTemplateDescription(BlankTemplateName, BlankDescription, PluginBaseDir / TEXT("Templates") / TEXT("Blank"), true, EHostType::Runtime, ELoadingPhase::Default, false));
+	BlankTemplate->SortPriority = 2;
+	TemplateDefinitions.Add(BlankTemplate);
 
-		TemplateDefinitions.Add(MakeShareable(new FPluginTemplateDescription(BlueprintLibTemplateName, BlueprintLibDescription, PluginBaseDir / TEXT("Templates") / TEXT("BlueprintLibrary"), true, EHostType::Runtime, ELoadingPhase::PreLoadingScreen)));
-		TemplateDefinitions.Add(MakeShareable(new FPluginTemplateDescription(BasicTemplateName, BasicDescription, PluginBaseDir / TEXT("Templates") / TEXT("Basic"), false, EHostType::Editor)));
-		TemplateDefinitions.Add(MakeShareable(new FPluginTemplateDescription(AdvancedTemplateName, AdvancedDescription, PluginBaseDir / TEXT("Templates") / TEXT("Advanced"), false, EHostType::Editor)));
-		TemplateDefinitions.Add(MakeShareable(new FPluginTemplateDescription(EditorModeTemplateName, EditorModeDescription, PluginBaseDir / TEXT("Templates") / TEXT("EditorMode"), false, EHostType::Editor)));
-		TemplateDefinitions.Add(MakeShareable(new FPluginTemplateDescription(ThirdPartyTemplateName, ThirdPartyDescription, PluginBaseDir / TEXT("Templates") / TEXT("ThirdPartyLibrary"), true, EHostType::Runtime)));
-	}
+	TemplateDefinitions.Add(MakeShareable(new FPluginTemplateDescription(BlueprintLibTemplateName, BlueprintLibDescription, PluginBaseDir / TEXT("Templates") / TEXT("BlueprintLibrary"), true, EHostType::Runtime, ELoadingPhase::PreLoadingScreen, false)));
+	TemplateDefinitions.Add(MakeShareable(new FPluginTemplateDescription(BasicTemplateName, BasicDescription, PluginBaseDir / TEXT("Templates") / TEXT("Basic"), false, EHostType::Editor, ELoadingPhase::Default, false)));
+	TemplateDefinitions.Add(MakeShareable(new FPluginTemplateDescription(AdvancedTemplateName, AdvancedDescription, PluginBaseDir / TEXT("Templates") / TEXT("Advanced"), false, EHostType::Editor, ELoadingPhase::Default, false)));
+	TemplateDefinitions.Add(MakeShareable(new FPluginTemplateDescription(EditorModeTemplateName, EditorModeDescription, PluginBaseDir / TEXT("Templates") / TEXT("EditorMode"), false, EHostType::Editor, ELoadingPhase::Default, false)));
+	TemplateDefinitions.Add(MakeShareable(new FPluginTemplateDescription(ThirdPartyTemplateName, ThirdPartyDescription, PluginBaseDir / TEXT("Templates") / TEXT("ThirdPartyLibrary"), true, EHostType::Runtime, ELoadingPhase::Default, false)));
 
 	// Add external templates that came from the modular feature interface (e.g., from another plugin like Game Features)
 	TemplateDefinitions.Append(FPluginBrowserModule::Get().GetAddedPluginTemplates());
diff --git a/Engine/Plugins/Editor/PluginBrowser/Source/PluginBrowser/Private/SNewPluginWizard.cpp b/Engine/Plugins/Editor/PluginBrowser/Source/PluginBrowser/Private/SNewPluginWizard.cpp
--- a/Engine/Plugins/Editor/PluginBrowser/Source/PluginBrowser/Private/SNewPluginWizard.cpp
+++ b/Engine/Plugins/Editor/PluginBrowser/Source/PluginBrowser/Private/SNewPluginWizard.cpp
@@ -363,10 +363,13 @@ TSharedRef<ITableRow> SNewPluginWizard::OnGenerateTemplateRow(TSharedRef<FPlugin
 {
 	const float PaddingAmount = FPluginStyle::Get()->GetFloat("PluginTile.Padding");
 	const float ThumbnailImageSize = FPluginStyle::Get()->GetFloat("PluginTile.ThumbnailImageSize");
+	const bool bTemplateEnabled = InItem->bSupportsContentOnlyProjects || !IsContentOnlyProject();
 
 	GeneratePluginTemplateDynamicBrush(InItem);
 
 	return SNew(STableRow< TSharedRef<FPluginTemplateDescription> >, OwnerTable)
+		.IsEnabled(bTemplateEnabled)
+		.ToolTipText_Lambda([bTemplateEnabled] { return !bTemplateEnabled ? LOCTEXT("TemplateDisabled", "This template is not available in Blueprint-only projects.") : FText::GetEmpty(); })
 		[
 			SNew(SBorder)
 			.BorderImage(FAppStyle::GetBrush("NoBorder"))
diff --git a/Engine/Source/Editor/UnrealEd/Public/Features/IPluginsEditorFeature.h b/Engine/Source/Editor/UnrealEd/Public/Features/IPluginsEditorFeature.h
--- a/Engine/Source/Editor/UnrealEd/Public/Features/IPluginsEditorFeature.h
+++ b/Engine/Source/Editor/UnrealEd/Public/Features/IPluginsEditorFeature.h
@@ -44,6 +44,9 @@ struct FPluginTemplateDescription
 	/** What is the expected Loading Phase for this plugin? */
 	ELoadingPhase::Type LoadingPhase;
 
+	/** Can the plugin be used in a content only project? */
+	bool bSupportsContentOnlyProjects;
+
 	/** Called just before the plugin is created */
 	virtual void CustomizeDescriptorBeforeCreation(FPluginDescriptor& Descriptor) {}
 
@@ -70,13 +73,14 @@ struct FPluginTemplateDescription
 	virtual void UpdatePluginNameTextWhenTemplateUnselected(FText& OutPluginNameText) {}
 
 	/** Constructor */
-	FPluginTemplateDescription(FText InName, FText InDescription, FString InOnDiskPath, bool InCanContainContent, EHostType::Type InModuleDescriptorType, ELoadingPhase::Type InLoadingPhase = ELoadingPhase::Default)
+	FPluginTemplateDescription(FText InName, FText InDescription, FString InOnDiskPath, bool InCanContainContent, EHostType::Type InModuleDescriptorType, ELoadingPhase::Type InLoadingPhase = ELoadingPhase::Default, bool InSupportsContentOnlyProjects = true)
 		: Name(InName)
 		, Description(InDescription)
 		, OnDiskPath(InOnDiskPath)
 		, bCanContainContent(InCanContainContent)
 		, ModuleDescriptorType(InModuleDescriptorType)
 		, LoadingPhase(InLoadingPhase)
+		, bSupportsContentOnlyProjects(InSupportsContentOnlyProjects)
 	{
 	}
 
