Plugins Browser - Made IsExperimental take precedence over IsBeta in the Plugin Browser UI display.

[RN] Plugins Browser - Fixed the display of Experimental plugins when a plugin has accidentally also been set as Beta.
#tests Checked loaded plugins in the Plugins window with both the IsBeta and IsExperimental flags set to true. Checked that Beta plugins still show correctly too.

diff --git a/Engine/Plugins/Editor/PluginBrowser/Source/PluginBrowser/Private/SPluginTile.cpp b/Engine/Plugins/Editor/PluginBrowser/Source/PluginBrowser/Private/SPluginTile.cpp
--- a/Engine/Plugins/Editor/PluginBrowser/Source/PluginBrowser/Private/SPluginTile.cpp
+++ b/Engine/Plugins/Editor/PluginBrowser/Source/PluginBrowser/Private/SPluginTile.cpp
@@ -434,13 +434,13 @@ void SPluginTile::RecreateWidgets()
 											[
 
 												SNew(SBorder)
-												.BorderImage(FPluginStyle::Get()->GetBrush(PluginDescriptor.bIsBetaVersion ? "PluginTile.BetaBorderImage" : "PluginTile.ExperimentalBorderImage"))
+												.BorderImage(FPluginStyle::Get()->GetBrush(PluginDescriptor.bIsExperimentalVersion ? "PluginTile.ExperimentalBorderImage" : "PluginTile.BetaBorderImage"))
 												.Visibility((PluginDescriptor.bIsBetaVersion || PluginDescriptor.bIsExperimentalVersion) ? EVisibility::Visible : EVisibility::Collapsed)
 												.Padding(FMargin(8.f, 1.f, 8.f, 2.f))
 												[
 													SNew(STextBlock)
 													.TextStyle(FPluginStyle::Get(), "PluginTile.BetaText")
-													.Text(PluginDescriptor.bIsBetaVersion ? LOCTEXT("PluginBetaVersionText", "Beta") : LOCTEXT("PluginExperimentalVersionText", "Experimental"))
+													.Text(PluginDescriptor.bIsExperimentalVersion ? LOCTEXT("PluginExperimentalVersionText", "Experimental") : LOCTEXT("PluginBetaVersionText", "Beta"))
 													.ToolTipText(this, &SPluginTile::GetBetaOrExperimentalHelpText)
 												]
 											]
@@ -572,17 +572,17 @@ void SPluginTile::OnEnablePluginCheckboxChanged(ECheckBoxState NewCheckedState)
 	if (bNewEnabledState)
 	{
 		// If this is plugin is marked as beta, make sure the user is aware before enabling it.
-		if (PluginDescriptor.bIsBetaVersion)
+		if (PluginDescriptor.bIsExperimentalVersion)
 		{
-			FText WarningMessage = FText::Format(LOCTEXT("Warning_EnablingBetaPlugin", "Plugin '{0}' is a beta version. {1} Are you sure you want to enable the plugin?"), GetPluginNameText(), GetBetaOrExperimentalHelpText());
+			FText WarningMessage = FText::Format(LOCTEXT("Warning_EnablingExperimentalPlugin", "Plugin '{0}' is an experimental version. {1} Are you sure you want to enable the plugin?"), GetPluginNameText(), GetBetaOrExperimentalHelpText());
 			if (EAppReturnType::No == FMessageDialog::Open(EAppMsgType::YesNo, WarningMessage))
 			{
 				return;
 			}
 		}
-		else if (PluginDescriptor.bIsExperimentalVersion)
+		else if (PluginDescriptor.bIsBetaVersion)
 		{
-			FText WarningMessage = FText::Format(LOCTEXT("Warning_EnablingExperimentalPlugin", "Plugin '{0}' is an experimental version. {1} Are you sure you want to enable the plugin?"), GetPluginNameText(), GetBetaOrExperimentalHelpText());
+			FText WarningMessage = FText::Format(LOCTEXT("Warning_EnablingBetaPlugin", "Plugin '{0}' is a beta version. {1} Are you sure you want to enable the plugin?"), GetPluginNameText(), GetBetaOrExperimentalHelpText());
 			if (EAppReturnType::No == FMessageDialog::Open(EAppMsgType::YesNo, WarningMessage))
 			{
 				return;
@@ -710,14 +710,14 @@ FText SPluginTile::GetBetaOrExperimentalHelpText() const
 	}
 
 	const FPluginDescriptor& PluginDescriptor = Plugin->GetDescriptor();
-	if (PluginDescriptor.bIsBetaVersion)
-	{
-		return LOCTEXT("Description_BetaPlugin", "Epic recommends using caution when shipping projects with beta plugins. Beta plugins support backwards compatibility for assets and APIs, but performance, stability, and platform support may not be shipping quality.");
-	}
 	if (PluginDescriptor.bIsExperimentalVersion)
 	{
 		return LOCTEXT("Description_ExperimentalPlugin", "Epic does not recommend shipping projects with experimental plugins. APIs, features, and the plugin itself are subject to change or be removed without notice.");
 	}
+	if (PluginDescriptor.bIsBetaVersion)
+	{
+		return LOCTEXT("Description_BetaPlugin", "Epic recommends using caution when shipping projects with beta plugins. Beta plugins support backwards compatibility for assets and APIs, but performance, stability, and platform support may not be shipping quality.");
+	}
 
 	return FText();
 }
