Plugin Browser: Allow override of "restart required" button behavior.

#rnx

diff --git a/Engine/Plugins/Editor/PluginBrowser/Source/PluginBrowser/Private/PluginBrowserModule.cpp b/Engine/Plugins/Editor/PluginBrowser/Source/PluginBrowser/Private/PluginBrowserModule.cpp
--- a/Engine/Plugins/Editor/PluginBrowser/Source/PluginBrowser/Private/PluginBrowserModule.cpp
+++ b/Engine/Plugins/Editor/PluginBrowser/Source/PluginBrowser/Private/PluginBrowserModule.cpp
@@ -23,6 +23,7 @@
 #include "SPluginPaths.h"
 #include "ToolMenu.h"
 #include "ToolMenus.h"
+#include "UnrealEdMisc.h"
 #include "Widgets/Docking/SDockTab.h"
 #include "Widgets/Notifications/SNotificationList.h"
 
@@ -347,7 +348,22 @@ TSharedRef<SDockTab> FPluginBrowserModule::HandleSpawnPluginBrowserTab(const FSp
 		SNew( SDockTab )
 		.TabRole( ETabRole::MajorTab );
 
-	MajorTab->SetContent( SNew( SPluginBrowser ) );
+	MajorTab->SetContent(
+		SNew( SPluginBrowser )
+		.OnRestartClicked_Lambda([this]() -> FReply
+		{
+			if (OnRestartClickedDelegate.IsBound())
+			{
+				OnRestartClickedDelegate.Execute();
+			}
+			else
+			{
+				const bool bWarn_false = false;
+				FUnrealEdMisc::Get().RestartEditor(bWarn_false);
+			}
+			return FReply::Handled();
+		})
+	);
 
 	PluginBrowserTab = MajorTab;
 	UpdatePreviousInstalledPlugins();
diff --git a/Engine/Plugins/Editor/PluginBrowser/Source/PluginBrowser/Private/PluginBrowserModule.h b/Engine/Plugins/Editor/PluginBrowser/Source/PluginBrowser/Private/PluginBrowserModule.h
--- a/Engine/Plugins/Editor/PluginBrowser/Source/PluginBrowser/Private/PluginBrowserModule.h
+++ b/Engine/Plugins/Editor/PluginBrowser/Source/PluginBrowser/Private/PluginBrowserModule.h
@@ -45,6 +45,7 @@ public:
 
 	virtual FOnLaunchReferenceViewer& OnLaunchReferenceViewerDelegate() override { return LaunchReferenceViewerDelegate; }
 	virtual FOnPluginDirectoriesChanged& OnPluginDirectoriesChanged() override { return OnPluginDirectoriesChangedDelegate; }
+	virtual FSimpleDelegate& OnRestartClicked() override { return OnRestartClickedDelegate; }
 
 	/**
 	 * Sets whether a plugin is pending enable/disable
@@ -117,6 +118,9 @@ private:
 	/** Register menu extensions for the content browser */
 	void AddContentBrowserMenuExtensions();
 
+	/** Delegate to call when the restart button in the pending restart notice is clicked. */
+	FSimpleDelegate OnRestartClickedDelegate;
+
 	/** List of added plugin templates */
 	TArray<TSharedRef<FPluginTemplateDescription>> AddedPluginTemplates;
 
diff --git a/Engine/Plugins/Editor/PluginBrowser/Source/PluginBrowser/Private/SPluginBrowser.cpp b/Engine/Plugins/Editor/PluginBrowser/Source/PluginBrowser/Private/SPluginBrowser.cpp
--- a/Engine/Plugins/Editor/PluginBrowser/Source/PluginBrowser/Private/SPluginBrowser.cpp
+++ b/Engine/Plugins/Editor/PluginBrowser/Source/PluginBrowser/Private/SPluginBrowser.cpp
@@ -8,7 +8,6 @@
 #include "SPluginTileList.h"
 #include "Widgets/SToolTip.h"
 #include "Framework/Docking/TabManager.h"
-#include "UnrealEdMisc.h"
 #include "Interfaces/IPluginManager.h"
 #include "PluginStyle.h"
 #include "Widgets/Navigation/SBreadcrumbTrail.h"
@@ -40,6 +39,8 @@ SPluginBrowser::~SPluginBrowser()
 
 void SPluginBrowser::Construct( const FArguments& Args )
 {
+	OnRestartClicked = Args._OnRestartClicked;
+
 	// Get the root directories which contain plugins
 	TArray<FString> WatchDirectoryNames;
 	WatchDirectoryNames.Add(FPaths::EnginePluginsDir());
@@ -319,12 +320,12 @@ void SPluginBrowser::Construct( const FArguments& Args )
 				.Padding(FMargin(18.0f, 20.0f, 18.0f, 16.0f))
 				[
 					SNew(SWarningOrErrorBox)
-					.Visibility(this, &SPluginBrowser::HandleRestartEditorNoticeVisibility)
+					.Visibility(this, &SPluginBrowser::HandleRestartNoticeVisibility)
 					.MessageStyle(EMessageStyle::Warning)
 					.Message(LOCTEXT("PluginSettingsRestartNotice", "You must restart Unreal Editor for your changes to take effect."))
 					[
 						SNew(SButton)
-						.OnClicked(this, &SPluginBrowser::HandleRestartEditorButtonClicked)
+						.OnClicked(this, &SPluginBrowser::HandleRestartButtonClicked)
 						.TextStyle(FAppStyle::Get(), "NormalText")
 						.Text(LOCTEXT("PluginSettingsRestartEditor", "Restart Now"))
 					]
@@ -350,16 +351,19 @@ void SPluginBrowser::Tick(const FGeometry& AllottedGeometry, const double InCurr
 	}
 }
 
-EVisibility SPluginBrowser::HandleRestartEditorNoticeVisibility() const
+EVisibility SPluginBrowser::HandleRestartNoticeVisibility() const
 {
 	return FPluginBrowserModule::Get().ShowPendingRestart() ? EVisibility::Visible : EVisibility::Collapsed;
 }
 
-FReply SPluginBrowser::HandleRestartEditorButtonClicked() const
+FReply SPluginBrowser::HandleRestartButtonClicked() const
 {
-	const bool bWarn = false;
-	FUnrealEdMisc::Get().RestartEditor(bWarn);
-	return FReply::Handled();
+	if (ensure(OnRestartClicked.IsBound()))
+	{
+		return OnRestartClicked.Execute();
+	}
+
+	return FReply::Unhandled();
 }
 
 void SPluginBrowser::SearchBox_OnPluginSearchTextChanged( const FText& NewText )
diff --git a/Engine/Plugins/Editor/PluginBrowser/Source/PluginBrowser/Private/SPluginBrowser.h b/Engine/Plugins/Editor/PluginBrowser/Source/PluginBrowser/Private/SPluginBrowser.h
--- a/Engine/Plugins/Editor/PluginBrowser/Source/PluginBrowser/Private/SPluginBrowser.h
+++ b/Engine/Plugins/Editor/PluginBrowser/Source/PluginBrowser/Private/SPluginBrowser.h
@@ -17,10 +17,15 @@ class SPluginCategoryTree;
  */
 class SPluginBrowser : public SCompoundWidget
 {
+public:
+	DECLARE_DELEGATE_RetVal(FReply, FReplyDelegate);
+
 	SLATE_BEGIN_ARGS( SPluginBrowser )
 	{
 	}
 
+		SLATE_EVENT(FReplyDelegate, OnRestartClicked)
+
 	SLATE_END_ARGS()
 
 	/** Destructor */
@@ -58,10 +63,10 @@ private:
 	EActiveTimerReturnType UpdatePluginsTimerCallback(double InCurrentTime, float InDeltaTime);
 
 	/** @return Is the "restart required" notice visible? */
-	EVisibility HandleRestartEditorNoticeVisibility() const;
+	EVisibility HandleRestartNoticeVisibility() const;
 
 	/** Handle the "restart now" button being clicked */
-	FReply HandleRestartEditorButtonClicked() const;
+	FReply HandleRestartButtonClicked() const;
 
 	/** Called when the text in the search box was changed */
 	void SearchBox_OnPluginSearchTextChanged( const FText& NewText );
@@ -79,6 +84,9 @@ private:
 	FReply HandleNewPluginButtonClicked() const;
 
 private:
+	/** Delegate to call when the restart button is clicked. */
+	FReplyDelegate OnRestartClicked;
+
 	/** Handles to the directory changed delegates */
 	TMap<FString, FDelegateHandle> WatchDirectories;
 
diff --git a/Engine/Plugins/Editor/PluginBrowser/Source/PluginBrowser/Public/IPluginBrowser.h b/Engine/Plugins/Editor/PluginBrowser/Source/PluginBrowser/Public/IPluginBrowser.h
--- a/Engine/Plugins/Editor/PluginBrowser/Source/PluginBrowser/Public/IPluginBrowser.h
+++ b/Engine/Plugins/Editor/PluginBrowser/Source/PluginBrowser/Public/IPluginBrowser.h
@@ -54,6 +54,9 @@ public:
 
 	/** Called when the external plugin directories configuration is modified via the browser. */
 	virtual FOnPluginDirectoriesChanged& OnPluginDirectoriesChanged() = 0;
+
+	/** Can optionally be bound to override the "restart required" button behavior. */
+	virtual FSimpleDelegate& OnRestartClicked() = 0;
 };
 
 
diff --git a/Engine/Source/Editor/UnrealEd/Public/UnrealEdMisc.h b/Engine/Source/Editor/UnrealEd/Public/UnrealEdMisc.h
--- a/Engine/Source/Editor/UnrealEd/Public/UnrealEdMisc.h
+++ b/Engine/Source/Editor/UnrealEd/Public/UnrealEdMisc.h
@@ -276,7 +276,7 @@ public:
 	UNREALED_API bool EnableWorldComposition(UWorld* InWorld, bool bEnable);
 
 	/** Get the path to the executable that runs the editor */
-	UNREALED_API FString GetProjectEditorBinaryPath();
+	UNREALED_API static FString GetProjectEditorBinaryPath();
 
 	/** Finds a map using only the map name, no extension, no path, also caches it for faster lookup next time. */
 	UNREALED_API FString FindMapFileFromPartialName(const FString& PartialMapName);
