Don't suppress prompt to enable newly installed plugins when the plugin browser is open
- I thought it was broken because I wasn't getting prompted. I happened to have the plugin browser open in my previous session and it was docked out of site.
- There's not a strong reason to suppress the prompt. Even if the browser is open, the prompt occurs as soon as the editor is opened so the user is not meaningfully interacting with it and the tab is basically guaranteed to have been automatically opened because it was open in the previous session.
- There's a better version of this where we open the browser directly to the "Installed" section, which highlights new plugins with a badge, but the plugin browser does not expose API for this currently.

diff --git a/Engine/Plugins/Editor/PluginBrowser/Source/PluginBrowser/Private/PluginBrowserModule.cpp b/Engine/Plugins/Editor/PluginBrowser/Source/PluginBrowser/Private/PluginBrowserModule.cpp
--- a/Engine/Plugins/Editor/PluginBrowser/Source/PluginBrowser/Private/PluginBrowserModule.cpp
+++ b/Engine/Plugins/Editor/PluginBrowser/Source/PluginBrowser/Private/PluginBrowserModule.cpp
@@ -365,7 +365,6 @@ TSharedRef<SDockTab> FPluginBrowserModule::HandleSpawnPluginBrowserTab(const FSp
 		})
 	);
 
-	PluginBrowserTab = MajorTab;
 	UpdatePreviousInstalledPlugins();
 
 	return MajorTab;
@@ -404,7 +403,6 @@ void FPluginBrowserModule::OnMainFrameLoaded(TSharedPtr<SWindow> InRootWindow, b
 	// Show a popup notification that allows the user to enable any new plugins
 	if(!bIsRunningStartupDialog 
 		&& NewlyInstalledPlugins.Num() > 0
-		&& !PluginBrowserTab.IsValid()
 		&& FGlobalTabmanager::Get()->GetTabPermissionList()->PassesFilter(PluginsEditorTabName))
 	{
 		FNotificationInfo Info(LOCTEXT("NewPluginsPopupTitle", "New plugins are available"));
diff --git a/Engine/Plugins/Editor/PluginBrowser/Source/PluginBrowser/Private/PluginBrowserModule.h b/Engine/Plugins/Editor/PluginBrowser/Source/PluginBrowser/Private/PluginBrowserModule.h
--- a/Engine/Plugins/Editor/PluginBrowser/Source/PluginBrowser/Private/PluginBrowserModule.h
+++ b/Engine/Plugins/Editor/PluginBrowser/Source/PluginBrowser/Private/PluginBrowserModule.h
@@ -128,9 +128,6 @@ private:
 	TArray<TPair<FOnPluginBeingEdited, FPluginEditorExtensionHandle>> CustomizePluginEditingDelegates;
 	int32 EditorExtensionCounter = 0;
 
-	/** The spawned browser tab */
-	TWeakPtr<SDockTab> PluginBrowserTab;
-
 	/** List of plugins that are pending enable/disable */
 	TMap<FString, bool> PendingEnablePlugins;
 
