Expose the Activate syntax to be used in the plugin details view.

diff --git a/Engine/Plugins/Editor/PluginBrowser/Source/PluginBrowser/Private/PluginMetadataObject.cpp b/Engine/Plugins/Editor/PluginBrowser/Source/PluginBrowser/Private/PluginMetadataObject.cpp
--- a/Engine/Plugins/Editor/PluginBrowser/Source/PluginBrowser/Private/PluginMetadataObject.cpp
+++ b/Engine/Plugins/Editor/PluginBrowser/Source/PluginBrowser/Private/PluginMetadataObject.cpp
@@ -19,6 +19,7 @@ void FPluginReferenceMetadata::PopulateFromDescriptor(const FPluginReferenceDesc
 	Name = InDescriptor.Name;
 	bEnabled = InDescriptor.bEnabled;
 	bOptional = InDescriptor.bOptional;
+	bActivate = InDescriptor.bActivate;
 }
 
 void FPluginReferenceMetadata::CopyIntoDescriptor(FPluginReferenceDescriptor& OutDescriptor) const
@@ -26,6 +27,7 @@ void FPluginReferenceMetadata::CopyIntoDescriptor(FPluginReferenceDescriptor& Ou
 	OutDescriptor.Name = Name;
 	OutDescriptor.bEnabled = bEnabled;
 	OutDescriptor.bOptional = bOptional;
+	OutDescriptor.bActivate = bActivate;
 }
 
 void FPluginDisallowedMetadata::PopulateFromDescriptor(const FPluginDisallowedDescriptor& InDescriptor)
diff --git a/Engine/Plugins/Editor/PluginBrowser/Source/PluginBrowser/Private/PluginMetadataObject.h b/Engine/Plugins/Editor/PluginBrowser/Source/PluginBrowser/Private/PluginMetadataObject.h
--- a/Engine/Plugins/Editor/PluginBrowser/Source/PluginBrowser/Private/PluginMetadataObject.h
+++ b/Engine/Plugins/Editor/PluginBrowser/Source/PluginBrowser/Private/PluginMetadataObject.h
@@ -33,6 +33,10 @@ public:
 	UPROPERTY(EditAnywhere, Category = "Plugin Reference")
 	bool bEnabled = true;
 
+	/** Whether the dependency plugin should be activated when this plugin is activated. This only works if the dependency is a GameFeaturePlugin */
+	UPROPERTY(EditAnywhere, Category = "Plugin Reference")
+	bool bActivate = false;
+
 	/**
 	 * Populate the fields of this object from an existing descriptor.
 	 */
