Add a way to sort plugins by names to simplify merging:
- button to do it manually in the Plugin tool
- editor settings to do it automatically when closing the plugin tool (bAutoSortPluginDependencies)

#Synced-CL 51771519
#Synced-CL 51771519

diff --git a/Engine/Plugins/Editor/PluginBrowser/Source/PluginBrowser/Private/PluginDescriptorEditor.cpp b/Engine/Plugins/Editor/PluginBrowser/Source/PluginBrowser/Private/PluginDescriptorEditor.cpp
--- a/Engine/Plugins/Editor/PluginBrowser/Source/PluginBrowser/Private/PluginDescriptorEditor.cpp
+++ b/Engine/Plugins/Editor/PluginBrowser/Source/PluginBrowser/Private/PluginDescriptorEditor.cpp
@@ -17,6 +17,7 @@
 #include "PluginStyle.h"
 #include "PluginMetadataObject.h"
 #include "PropertyEditorModule.h"
+#include "Misc/ConfigCacheIni.h"
 #include "Misc/ScopedSlowTask.h"
 
 #define LOCTEXT_NAMESPACE "PluginListTile"
@@ -88,6 +89,14 @@ FReply FPluginDescriptorEditor::OnEditPluginFinished(UPluginMetadataObject* Meta
 {
 	FPluginDescriptor OldDescriptor = Plugin->GetDescriptor();
 
+	bool bAutoSortPluginDependencies = false;
+	if (GConfig
+		&& GConfig->GetBool(TEXT("PluginDescriptorSettings"), TEXT("bAutoSortPluginDependencies"), bAutoSortPluginDependencies, GEditorIni)
+		&& bAutoSortPluginDependencies)
+	{
+		MetadataObject->SortDependencies();
+	}
+
 	// Update the descriptor with the new metadata
 	FPluginDescriptor NewDescriptor = OldDescriptor;
 	MetadataObject->CopyIntoDescriptor(NewDescriptor);
diff --git a/Engine/Plugins/Editor/PluginBrowser/Source/PluginBrowser/Private/PluginMetadataObject.cpp b/Engine/Plugins/Editor/PluginBrowser/Source/PluginBrowser/Private/PluginMetadataObject.cpp
--- a/Engine/Plugins/Editor/PluginBrowser/Source/PluginBrowser/Private/PluginMetadataObject.cpp
+++ b/Engine/Plugins/Editor/PluginBrowser/Source/PluginBrowser/Private/PluginMetadataObject.cpp
@@ -46,6 +46,18 @@ UPluginMetadataObject::UPluginMetadataObject(const FObjectInitializer& ObjectIni
 {
 }
 
+void UPluginMetadataObject::SortDependencies()
+{
+	Plugins.StableSort([] (const FPluginReferenceMetadata& First, const FPluginReferenceMetadata& Second)
+		{
+			return First.Name < Second.Name;
+		});
+	DisallowedPlugins.StableSort([] (const FPluginDisallowedMetadata& First, const FPluginDisallowedMetadata& Second)
+		{
+			return First.Name < Second.Name;
+		});
+}
+
 void UPluginMetadataObject::PopulateFromPlugin(TSharedPtr<IPlugin> InPlugin)
 {
 	SourcePlugin = InPlugin;
diff --git a/Engine/Plugins/Editor/PluginBrowser/Source/PluginBrowser/Private/PluginMetadataObject.h b/Engine/Plugins/Editor/PluginBrowser/Source/PluginBrowser/Private/PluginMetadataObject.h
--- a/Engine/Plugins/Editor/PluginBrowser/Source/PluginBrowser/Private/PluginMetadataObject.h
+++ b/Engine/Plugins/Editor/PluginBrowser/Source/PluginBrowser/Private/PluginMetadataObject.h
@@ -171,6 +171,10 @@ public:
 	/** Editing extensions */
 	TArray<TSharedPtr<FPluginEditorExtension>> Extensions;
 
+	/** Sort the plugins and disallowed plugins references in ascending name order */
+	UFUNCTION(CallInEditor, Category=Dependencies, meta = (DisplayName = "Sort dependencies"))
+	void SortDependencies();
+
 	/**
 	 * Populate the fields of this object from an existing descriptor.
 	 */
