Fixup plugin wizard to check for the actual name of the plugin we'll create rather than the name the user entered. This fixes an issue where plugin templates can add a pre or postfix to the PLUGIN_NAME of the .uplugin for the template.

#Synced-CL 52439769

diff --git a/Engine/Plugins/Editor/PluginBrowser/Source/PluginBrowser/Private/SNewPluginWizard.cpp b/Engine/Plugins/Editor/PluginBrowser/Source/PluginBrowser/Private/SNewPluginWizard.cpp
--- a/Engine/Plugins/Editor/PluginBrowser/Source/PluginBrowser/Private/SNewPluginWizard.cpp
+++ b/Engine/Plugins/Editor/PluginBrowser/Source/PluginBrowser/Private/SNewPluginWizard.cpp
@@ -600,7 +600,7 @@ void SNewPluginWizard::ValidateFullPluginPath()
 
 	// Check for issues with name. Fail silently if text is empty
 	FText PluginNameError;
-	bIsPluginNameValid = !GetCurrentPluginName().IsEmpty() && FPluginUtils::ValidateNewPluginNameAndLocation(GetCurrentPluginName().ToString(), PluginFolderPath, &PluginNameError);
+	bIsPluginNameValid = !GetCurrentPluginName().IsEmpty() && FPluginUtils::ValidateNewPluginNameAndLocation(GetActualPluginName(), PluginFolderPath, &PluginNameError);
 	FilePathBlock->SetNameError(PluginNameError);
 }
 
@@ -619,6 +619,25 @@ FText SNewPluginWizard::GetCurrentPluginName() const
 	return PluginNameText;
 }
 
+FString SNewPluginWizard::GetActualPluginName() const
+{
+	TSharedPtr<FPluginTemplateDescription> Template = PluginWizardDefinition->GetSelectedTemplate();
+	if (Template.IsValid())
+	{
+		TArray<FString> UPluginFiles;
+		IFileManager::Get().FindFiles(UPluginFiles, *Template->OnDiskPath, TEXT("*.uplugin"));
+		for (FString File : UPluginFiles)
+		{
+			File.RemoveFromEnd(TEXT(".uplugin"));
+			if (File.Contains(TEXT("PLUGIN_NAME")))
+			{
+				return File.Replace(TEXT("PLUGIN_NAME"), *GetCurrentPluginName().ToString(), ESearchCase::CaseSensitive);
+			}
+		}
+	}
+	return GetCurrentPluginName().ToString();
+}
+
 ECheckBoxState SNewPluginWizard::IsEnginePlugin() const
 {
 	return bIsSelectedPathInEngine ? ECheckBoxState::Checked : ECheckBoxState::Unchecked;
@@ -665,21 +684,7 @@ FReply SNewPluginWizard::OnCreatePluginClicked()
 		return FReply::Unhandled();
 	}
 
-	FString PluginName;
-	TArray<FString> UPluginFiles;
-	IFileManager::Get().FindFiles(UPluginFiles, *Template->OnDiskPath, TEXT("*.uplugin"));
-	for (FString& File : UPluginFiles)
-	{
-		File.RemoveFromEnd(TEXT(".uplugin"));
-		if (File.Contains(TEXT("PLUGIN_NAME")))
-		{
-			PluginName = File.Replace(TEXT("PLUGIN_NAME"), *GetCurrentPluginName().ToString(), ESearchCase::CaseSensitive);
-		}
-	}
-	if (PluginName.IsEmpty())
-	{
-		PluginName = GetCurrentPluginName().ToString();
-	}
+	FString PluginName = GetActualPluginName();
 	const bool bHasModules = PluginWizardDefinition->HasModules();
 	
 	FPluginUtils::FNewPluginParamsWithDescriptor CreationParams;
diff --git a/Engine/Plugins/Editor/PluginBrowser/Source/PluginBrowser/Private/SNewPluginWizard.h b/Engine/Plugins/Editor/PluginBrowser/Source/PluginBrowser/Private/SNewPluginWizard.h
--- a/Engine/Plugins/Editor/PluginBrowser/Source/PluginBrowser/Private/SNewPluginWizard.h
+++ b/Engine/Plugins/Editor/PluginBrowser/Source/PluginBrowser/Private/SNewPluginWizard.h
@@ -94,6 +94,13 @@ private:
 	/** Get the current name of the plugin */
 	FText GetCurrentPluginName() const;
 
+	/**
+	 * Get the actual name of the plugin that will be created, derived from the selected template's
+	 * .uplugin filename (replacing the PLUGIN_NAME token with the user-entered name). Falls back to
+	 * the user-entered name if the template does not use the token.
+	 */
+	FString GetActualPluginName() const;
+
 	/**
 	 * Whether we will create a plugin in the engine directory
 	 */
