function rename

diff --git a/Engine/Plugins/Editor/PluginBrowser/Source/PluginBrowser/Private/SNewPluginWizard.cpp b/Engine/Plugins/Editor/PluginBrowser/Source/PluginBrowser/Private/SNewPluginWizard.cpp
--- a/Engine/Plugins/Editor/PluginBrowser/Source/PluginBrowser/Private/SNewPluginWizard.cpp
+++ b/Engine/Plugins/Editor/PluginBrowser/Source/PluginBrowser/Private/SNewPluginWizard.cpp
@@ -600,7 +600,7 @@ void SNewPluginWizard::ValidateFullPluginPath()
 
 	// Check for issues with name. Fail silently if text is empty
 	FText PluginNameError;
-	bIsPluginNameValid = !GetCurrentPluginName().IsEmpty() && FPluginUtils::ValidateNewPluginNameAndLocation(GetActualPluginName(), PluginFolderPath, &PluginNameError);
+	bIsPluginNameValid = !GetCurrentPluginName().IsEmpty() && FPluginUtils::ValidateNewPluginNameAndLocation(GetExpandedPluginName(), PluginFolderPath, &PluginNameError);
 	FilePathBlock->SetNameError(PluginNameError);
 }
 
@@ -619,7 +619,7 @@ FText SNewPluginWizard::GetCurrentPluginName() const
 	return PluginNameText;
 }
 
-FString SNewPluginWizard::GetActualPluginName() const
+FString SNewPluginWizard::GetExpandedPluginName() const
 {
 	TSharedPtr<FPluginTemplateDescription> Template = PluginWizardDefinition->GetSelectedTemplate();
 	if (Template.IsValid())
@@ -684,7 +684,7 @@ FReply SNewPluginWizard::OnCreatePluginClicked()
 		return FReply::Unhandled();
 	}
 
-	FString PluginName = GetActualPluginName();
+	FString PluginName = GetExpandedPluginName();
 	const bool bHasModules = PluginWizardDefinition->HasModules();
 	
 	FPluginUtils::FNewPluginParamsWithDescriptor CreationParams;
diff --git a/Engine/Plugins/Editor/PluginBrowser/Source/PluginBrowser/Private/SNewPluginWizard.h b/Engine/Plugins/Editor/PluginBrowser/Source/PluginBrowser/Private/SNewPluginWizard.h
--- a/Engine/Plugins/Editor/PluginBrowser/Source/PluginBrowser/Private/SNewPluginWizard.h
+++ b/Engine/Plugins/Editor/PluginBrowser/Source/PluginBrowser/Private/SNewPluginWizard.h
@@ -99,7 +99,7 @@ private:
 	 * .uplugin filename (replacing the PLUGIN_NAME token with the user-entered name). Falls back to
 	 * the user-entered name if the template does not use the token.
 	 */
-	FString GetActualPluginName() const;
+	FString GetExpandedPluginName() const;
 
 	/**
 	 * Whether we will create a plugin in the engine directory
